[StateStream]
* First iteration of state stream system to be able to decouple game and render. Very much WIP.

diff --git a/Engine/Source/Programs/LowLevelTests/StateStreamTests/Public/MockStateStream.h b/Engine/Source/Programs/LowLevelTests/StateStreamTests/Public/MockStateStream.h
new file mode 100644
--- /dev/null
+++ b/Engine/Source/Programs/LowLevelTests/StateStreamTests/Public/MockStateStream.h
@@ -0,0 +1,65 @@
+// Copyright Epic Games, Inc. All Rights Reserved.
+
+#pragma once
+
+#include "TransformStateStream.h"
+#include "MockStateStream.generated.h"
+
+// This is the interface Gameplay is using to interact with statestream.
+
+////////////////////////////////////////////////////////////////////////////////////////////////////
+// Static state. Should be immutable during the life time of the instance. Should use attributes to generate serialization code etc using UHT.
+USTRUCT(StateStreamStaticState)
+struct FMockStaticState
+{
+	GENERATED_USTRUCT_BODY()
+};
+
+
+////////////////////////////////////////////////////////////////////////////////////////////////////
+// Dynamic state. Can change over the lifetime of the instance.
+USTRUCT(StateStreamDynamicState)
+struct FMockDynamicState
+{
+	GENERATED_USTRUCT_BODY()
+	FMockDynamicState(uint32 V) { SetValue(V); }
+	FMockDynamicState(uint32 V, bool B) { SetValue(V); SetBit2(B); }
+
+private:
+
+	UPROPERTY()
+	uint32 Value = 0;
+
+	UPROPERTY()
+	FTransformHandle Transform;
+
+	UPROPERTY()
+	uint32 bBit1 : 1;
+
+	UPROPERTY()
+	uint32 bBit2 : 1;
+};
+
+
+////////////////////////////////////////////////////////////////////////////////////////////////////
+// Handle to keep track of instances and their lifetime. Is ref counted so when count reaches 0 instance is tagged as deleted
+USTRUCT(StateStreamHandle)
+struct FMockHandle : public FStateStreamHandle
+{
+	GENERATED_USTRUCT_BODY()
+};
+
+
+////////////////////////////////////////////////////////////////////////////////////////////////////
+// Id. Used to identify statestreams. For finding statestream game side and register dependencies render side.
+inline constexpr uint32 MockStateStreamId = 128;
+
+
+////////////////////////////////////////////////////////////////////////////////////////////////////
+// The statestream itself that Gameplay is using to create instances.
+class IMockStateStream
+{
+public:
+	DECLARE_STATESTREAM(Mock)
+	virtual Handle Game_CreateInstance(const FMockStaticState& Ss, const FMockDynamicState& Ds) = 0;
+};
diff --git a/Engine/Source/Programs/LowLevelTests/StateStreamTests/Public/MockStateStreamImpl.h b/Engine/Source/Programs/LowLevelTests/StateStreamTests/Public/MockStateStreamImpl.h
new file mode 100644
--- /dev/null
+++ b/Engine/Source/Programs/LowLevelTests/StateStreamTests/Public/MockStateStreamImpl.h
@@ -0,0 +1,108 @@
+// Copyright Epic Games, Inc. All Rights Reserved.
+
+#pragma once
+
+#include "GenericStateStream.h"
+#include "MockStateStream.h"
+#include "TransformStateStreamImpl.h"
+
+// ============ This is ONLY RT side ==============
+
+////////////////////////////////////////////////////////////////////////////////////////////////////
+
+struct MockObject : public FTransformObjectListener
+{
+	virtual ~MockObject()
+	{
+		SetTransformObject({});
+	}
+
+	void SetTransformObject(const FTransformHandle& Handle)
+	{
+		if (Transform)
+		{
+			Transform->RemoveListener(this);
+		}
+		Transform = (FTransformObject*)Handle.Render_GetUserData();
+		if (Transform)
+		{
+			Transform->AddListener(this);
+		}
+	}
+
+	virtual void OnTransformObjectDirty() override
+	{
+		// Do stuff
+	}
+
+	float Value = 0;
+	bool Bit = false;
+	TRefCountPtr<FTransformObject> Transform;
+};
+
+////////////////////////////////////////////////////////////////////////////////////////////////////
+
+class FMockStateStreamImpl : public TStateStream<TStateStreamSettings<IMockStateStream, MockObject>>
+{
+public:
+	TArray<MockObject*> Instances;
+	uint32 CreateCount = 0;
+	uint32 CreateAndDestroyCount = 0;
+	uint32 UpdateCount = 0;
+	uint32 DestroyCount = 0;
+
+	virtual void Render_OnCreate(const FMockStaticState& Ss, const FMockDynamicState& Ds, MockObject*& UserData, bool IsDestroyedInSameFrame) override
+	{
+		if (IsDestroyedInSameFrame)
+		{
+			++CreateAndDestroyCount;
+			return;
+		}
+
+		MockObject* Object = new MockObject(); 
+		Object->Value = Ds.GetValue();
+		Object->Bit = Ds.GetBit2();
+		UserData = Object;
+
+		if (Ds.GetTransform().IsValid())
+		{
+			Object->SetTransformObject(Ds.GetTransform());
+		}
+
+		Instances.Add(Object);
+		++CreateCount;
+	}
+
+	virtual void Render_OnUpdate(const FMockStaticState& Ss, const FMockDynamicState& Ds, MockObject*& UserData) override
+	{
+		MockObject* Object = UserData;
+		if (Ds.ValueModified())
+		{
+			Object->Value = Ds.GetValue();
+		}
+		if (Ds.Bit2Modified())
+		{
+			Object->Bit = Ds.GetBit2();
+		}
+		if (Ds.TransformModified())
+		{
+			Object->SetTransformObject(Ds.GetTransform());
+		}
+		++UpdateCount;
+	}
+
+	virtual void Render_OnDestroy(const FMockStaticState& Ss, const FMockDynamicState& Ds, MockObject*& UserData) override
+	{
+		MockObject* Object = UserData;
+		if (!Object)
+		{
+			return;
+		}
+		Instances.Remove(Object);
+		delete Object;
+		UserData = nullptr;
+		++DestroyCount;
+	}
+};
+
+////////////////////////////////////////////////////////////////////////////////////////////////////
diff --git a/Engine/Source/Programs/LowLevelTests/StateStreamTests/StateStreamTests.Build.cs b/Engine/Source/Programs/LowLevelTests/StateStreamTests/StateStreamTests.Build.cs
new file mode 100644
--- /dev/null
+++ b/Engine/Source/Programs/LowLevelTests/StateStreamTests/StateStreamTests.Build.cs
@@ -0,0 +1,28 @@
+// Copyright Epic Games, Inc. All Rights Reserved.
+using UnrealBuildTool;
+
+public class StateStreamTests : TestModuleRules
+{
+	public StateStreamTests(ReadOnlyTargetRules Target) : base(Target, true)
+	{
+		TestMetadata = new Metadata()
+		{
+			TestName = "StateStream",
+			TestShortName = "StateStream",
+			ReportType = "xml",
+			SupportedPlatforms = {
+				UnrealTargetPlatform.Win64,
+				UnrealTargetPlatform.Linux,
+				UnrealTargetPlatform.Mac,
+				UnrealTargetPlatform.Android,
+				UnrealTargetPlatform.IOS }
+		};
+
+		PrivateDependencyModuleNames.AddRange(
+			new string[] {
+				"Core",
+				"CoreUObject",
+				"StateStream",
+			});
+	}
+}
diff --git a/Engine/Source/Programs/LowLevelTests/StateStreamTests/StateStreamTests.Target.cs b/Engine/Source/Programs/LowLevelTests/StateStreamTests/StateStreamTests.Target.cs
new file mode 100644
--- /dev/null
+++ b/Engine/Source/Programs/LowLevelTests/StateStreamTests/StateStreamTests.Target.cs
@@ -0,0 +1,18 @@
+// Copyright Epic Games, Inc. All Rights Reserved.
+
+using UnrealBuildTool;
+
+[SupportedPlatforms(UnrealPlatformClass.All)]
+public class StateStreamTestsTarget : TestTargetRules
+{
+	public StateStreamTestsTarget(TargetInfo Target) : base(Target)
+	{
+		bCompileWithPluginSupport = false;
+		bCompileAgainstCoreUObject = true;
+
+		bBuildWithEditorOnlyData = Target.Platform.IsInGroup(UnrealPlatformGroup.Desktop)
+			&& (Target.Configuration == UnrealTargetConfiguration.Debug || Target.Configuration == UnrealTargetConfiguration.Development);
+
+		GlobalDefinitions.Add("UE_STATESTREAM_TIME_TYPE=int64");
+	}
+}
diff --git a/Engine/Source/Programs/LowLevelTests/StateStreamTests/Tests/StateStreamTests.cpp b/Engine/Source/Programs/LowLevelTests/StateStreamTests/Tests/StateStreamTests.cpp
new file mode 100644
--- /dev/null
+++ b/Engine/Source/Programs/LowLevelTests/StateStreamTests/Tests/StateStreamTests.cpp
@@ -0,0 +1,572 @@
+// Copyright Epic Games, Inc. All Rights Reserved.
+#include "TestCommon/Initialization.h"
+#include "TestHarness.h"
+#include "MockStateStreamImpl.h"
+#include "StateStreamManagerImpl.h"
+#include "TransformStateStreamImpl.h"
+
+
+TEST_CASE("FPow2ChunkedArray", "[StateStream]")
+{
+	//   0 -  15 Bucket 0  (16 entries)
+	//  16 -  47 Bucket 1  (32 entries)
+	//  48 - 111 Bucket 2  (64 entries)
+	// 112 - 239 Bucket 3  (128 entries)
+	// 240 - 495 Bucket 4  (256 entries)
+
+	FPow2ChunkedArray<int> Store;
+	CHECK_EQUAL(Store.SkipCount, 4);
+	CHECK_EQUAL(Store.BucketCount, 21);
+	CHECK_EQUAL(Store.GetBucketIndex(0), 0);
+	CHECK_EQUAL(Store.GetBucketIndex(15), 0);
+	CHECK_EQUAL(Store.GetBucketIndex(16), 1);
+	CHECK_EQUAL(Store.GetBucketIndex(47), 1);
+	CHECK_EQUAL(Store.GetBucketIndex(48), 2);
+	CHECK_EQUAL(Store.GetBucketIndex(111), 2);
+	CHECK_EQUAL(Store.GetBucketIndex(112), 3);
+	CHECK_EQUAL(Store.GetBucketIndex(16777215), Store.BucketCount - 1);
+
+	CHECK_EQUAL(Store.GetBucketSize(0), 16);
+	CHECK_EQUAL(Store.GetBucketSize(1), 32);
+	CHECK_EQUAL(Store.GetBucketSize(2), 64);
+	CHECK_EQUAL(Store.GetBucketSize(3), 128);
+
+	CHECK_EQUAL(Store.GetBucketStart(0), 0);
+	CHECK_EQUAL(Store.GetBucketStart(1), 16);
+	CHECK_EQUAL(Store.GetBucketStart(2), 48);
+	CHECK_EQUAL(Store.GetBucketStart(3), 112);
+	CHECK_EQUAL(Store.GetBucketStart(4), 240);
+
+	for (int i=0;i!=1000;++i)
+	{
+		Store.Add(i);
+	}
+
+	for (int i=0;i!=1000;++i)
+	{
+		CHECK_EQUAL(Store[i], i);
+	}
+}
+
+TEST_CASE("StateStream::Store", "[StateStream]")
+{
+	TStateStreamStore<int> Store;
+
+	for (int i=0;i!=1000;++i)
+	{
+		Store.Add(i);
+	}
+	CHECK_EQUAL(Store.GetUsedCount(), 1000);
+
+	for (int i=0;i!=1000;++i)
+	{
+		CHECK_EQUAL(Store[i], i);
+	}
+
+	Store.Remove(100);
+	Store.Remove(101);
+
+	CHECK_EQUAL(Store.GetUsedCount(), 998);
+	CHECK_EQUAL(Store.Add(1234), 101);
+	CHECK_EQUAL(Store.Add(1235), 100);
+	CHECK_EQUAL(Store.GetUsedCount(), 1000);
+	CHECK_EQUAL(Store[100], 1235);
+	CHECK_EQUAL(Store[101], 1234);
+	CHECK_EQUAL(Store.Add(1235), 1000);
+	CHECK_EQUAL(Store.GetUsedCount(), 1001);
+	CHECK_EQUAL(Store[1000], 1235);
+}
+
+TEST_CASE("StateStream::SingleInstance", "[StateStream]")
+{
+	// Initialization (needs to be done centrally/renderside before gameplay starts) (These instances are not visible to GT)
+	FMockStateStreamImpl MockStateStreamImpl;
+	FStateStreamManagerImpl ManagerImpl;
+	ManagerImpl.Render_Register(MockStateStreamImpl, false);
+
+
+	// Game thread populating ticks
+	{
+		IStateStreamManager& Manager = ManagerImpl;
+		IMockStateStream& Stream = Manager.Game_Get<IMockStateStream>();
+
+		// Tick 1, create an instance
+		Manager.Game_BeginTick();
+		FMockHandle Handle = Stream.Game_CreateInstance(FMockStaticState(), FMockDynamicState{0});
+		Manager.Game_EndTick(100);
+
+		// Tick 2, update the instance. Set Value to 100
+		Manager.Game_BeginTick();
+		Handle.Update(FMockDynamicState{100, true});
+		Handle.Update(FMockDynamicState{100});
+		Manager.Game_EndTick(200);
+
+		// Tick 3, update the instance again. Set Value to 200
+		Manager.Game_BeginTick();
+		Handle.Update(FMockDynamicState{200, false});
+		Manager.Game_EndTick(300);
+
+		// Tick 4, destroy the instance
+		Manager.Game_BeginTick();
+		Handle = {};
+		Manager.Game_EndTick(400);
+
+		// Tick 5, empty
+		Manager.Game_BeginTick();
+		Manager.Game_EndTick(500);
+	}
+
+
+	// Render thread consuming ticks
+
+	FStateStreamManagerImpl& Manager = ManagerImpl;
+	FMockStateStreamImpl& Stream = MockStateStreamImpl;
+
+
+	SECTION("Rend not crossing ticks")
+    {
+		Manager.Render_Update(10); // We start consuming Tick 1
+		CHECK_EQUAL(Stream.Instances.Num(), 1);
+		CHECK_EQUAL(Stream.Instances[0]->Value, 0);
+		CHECK_EQUAL(Stream.Instances[0]->Bit, false);
+
+		Manager.Render_Update(50); // Update from Tick 1 (still no interpolation, instance was created in this tick)
+		CHECK_EQUAL(Stream.Instances.Num(), 1);
+		CHECK_EQUAL(Stream.Instances[0]->Value, 0);
+		CHECK_EQUAL(Stream.Instances[0]->Bit, false);
+
+		Manager.Render_Update(100); // Values should be exactly like 1
+		CHECK_EQUAL(Stream.Instances.Num(), 1);
+		CHECK_EQUAL(Stream.Instances[0]->Value, 0);
+		CHECK_EQUAL(Stream.Instances[0]->Bit, false);
+
+		CHECK_EQUAL(Stream.CreateCount, 1);
+		CHECK_EQUAL(Stream.CreateAndDestroyCount, 0);
+		CHECK_EQUAL(Stream.DestroyCount, 0);
+		CHECK_EQUAL(Stream.UpdateCount, 0);
+
+		// Moving into Tick 2
+
+		Manager.Render_Update(110); // 10 units into Tick 2.. Values should start interpolating from Tick 1 against Tick 2.
+		CHECK_EQUAL(Stream.Instances.Num(), 1);
+		CHECK_EQUAL(Stream.Instances[0]->Value, 10);
+		CHECK_EQUAL(Stream.Instances[0]->Bit, true);
+
+		Manager.Render_Update(150); // 50 units into Tick 2
+		CHECK_EQUAL(Stream.Instances.Num(), 1);
+		CHECK_EQUAL(Stream.Instances[0]->Value, 50);
+		CHECK_EQUAL(Stream.Instances[0]->Bit, true);
+
+		Manager.Render_Update(200); // Full Tick 2.. Value should be like Tick 2
+		CHECK_EQUAL(Stream.Instances.Num(), 1);
+		CHECK_EQUAL(Stream.Instances[0]->Value, 100);
+		CHECK_EQUAL(Stream.Instances[0]->Bit, true);
+
+		CHECK_EQUAL(Stream.CreateCount, 1);
+		CHECK_EQUAL(Stream.CreateAndDestroyCount, 0);
+		CHECK_EQUAL(Stream.DestroyCount, 0);
+		CHECK_EQUAL(Stream.UpdateCount, 3);
+
+		// Moving into Tick 3
+
+		Manager.Render_Update(210); // 10 units into Tick 3.. 
+		CHECK_EQUAL(Stream.Instances.Num(), 1);
+		CHECK_EQUAL(Stream.Instances[0]->Value, 110);
+		CHECK_EQUAL(Stream.Instances[0]->Bit, false);
+
+		Manager.Render_Update(280); // 80 units into Tick 3.. 
+		CHECK_EQUAL(Stream.Instances.Num(), 1);
+		CHECK_EQUAL(Stream.Instances[0]->Value, 180);
+
+		Manager.Render_Update(300); // Full Tick 3.. 
+		CHECK_EQUAL(Stream.Instances.Num(), 1);
+		CHECK_EQUAL(Stream.Instances[0]->Value, 200);
+
+		CHECK_EQUAL(Stream.CreateCount, 1);
+		CHECK_EQUAL(Stream.CreateAndDestroyCount, 0);
+		CHECK_EQUAL(Stream.DestroyCount, 0);
+		CHECK_EQUAL(Stream.UpdateCount, 6);
+
+		// Moving into Tick 4
+
+		Manager.Render_Update(310); // 10 units into Tick 4
+		CHECK_EQUAL(Stream.Instances.Num(), 0);
+
+		CHECK_EQUAL(Stream.CreateCount, 1);
+		CHECK_EQUAL(Stream.CreateAndDestroyCount, 0);
+		CHECK_EQUAL(Stream.DestroyCount, 1);
+		CHECK_EQUAL(Stream.UpdateCount, 6);
+    }
+
+
+	SECTION("Rend crossing single ticks")
+    {
+		// Consuming entire Tick 1 and moving into Tick 2
+
+		Manager.Render_Update(110);
+		CHECK_EQUAL(Stream.Instances.Num(), 1);
+		CHECK_EQUAL(Stream.Instances[0]->Value, 10);
+
+		CHECK_EQUAL(Stream.CreateCount, 1);
+		CHECK_EQUAL(Stream.CreateAndDestroyCount, 0);
+		CHECK_EQUAL(Stream.DestroyCount, 0);
+		CHECK_EQUAL(Stream.UpdateCount, 0);
+
+		// Consuming rest of Tick 2 and moving into Tick 3
+
+		Manager.Render_Update(210);
+		CHECK_EQUAL(Stream.Instances.Num(), 1);
+		CHECK_EQUAL(Stream.Instances[0]->Value, 110);
+
+		CHECK_EQUAL(Stream.CreateCount, 1);
+		CHECK_EQUAL(Stream.CreateAndDestroyCount, 0);
+		CHECK_EQUAL(Stream.DestroyCount, 0);
+		CHECK_EQUAL(Stream.UpdateCount, 1);
+
+		// Consuming rest of Tick 3 and moving into Tick 4
+
+		Manager.Render_Update(310);
+		CHECK_EQUAL(Stream.Instances.Num(), 0);
+
+		CHECK_EQUAL(Stream.CreateCount, 1);
+		CHECK_EQUAL(Stream.CreateAndDestroyCount, 0);
+		CHECK_EQUAL(Stream.DestroyCount, 1);
+		CHECK_EQUAL(Stream.UpdateCount, 1);
+	}
+
+
+	SECTION("Rend crossing multiple ticks")
+    {
+		// Consuming entire Tick 1 and 2 and moving into Tick 3
+
+		Manager.Render_Update(210);
+		CHECK_EQUAL(Stream.Instances.Num(), 1);
+		CHECK_EQUAL(Stream.Instances[0]->Value, 110);
+
+		CHECK_EQUAL(Stream.CreateCount, 1);
+		CHECK_EQUAL(Stream.CreateAndDestroyCount, 0);
+		CHECK_EQUAL(Stream.DestroyCount, 0);
+		CHECK_EQUAL(Stream.UpdateCount, 0);
+
+		// Consuming rest of Tick 3, entire 4 and entering 5
+
+		Manager.Render_Update(410);
+		CHECK_EQUAL(Stream.Instances.Num(), 0);
+
+		CHECK_EQUAL(Stream.CreateCount, 1);
+		CHECK_EQUAL(Stream.CreateAndDestroyCount, 0);
+		CHECK_EQUAL(Stream.DestroyCount, 1);
+		CHECK_EQUAL(Stream.UpdateCount, 0);
+	}
+
+
+	SECTION("Rend crossing all ticks")
+    {
+		// Consuming entire Tick 1 to 5
+
+		Manager.Render_Update(500);
+		CHECK_EQUAL(Stream.Instances.Num(), 0);
+		CHECK_EQUAL(Stream.CreateCount, 0);
+		CHECK_EQUAL(Stream.CreateAndDestroyCount, 1);
+		CHECK_EQUAL(Stream.DestroyCount, 0);
+		CHECK_EQUAL(Stream.UpdateCount, 0);
+	}
+}
+
+TEST_CASE("StateStream::MultipleInstances", "[StateStream]")
+{
+	FMockStateStreamImpl MockStateStreamImpl;
+	FStateStreamManagerImpl ManagerImpl;
+	ManagerImpl.Render_Register(MockStateStreamImpl, false);
+
+	{
+		IStateStreamManager& Manager = ManagerImpl;
+		IMockStateStream& Stream = Manager.Game_Get<IMockStateStream>();
+
+		// Tick 1
+		Manager.Game_BeginTick();
+		FMockHandle Handle1 = Stream.Game_CreateInstance(FMockStaticState(), FMockDynamicState{0});
+		FMockHandle Handle2 = Stream.Game_CreateInstance(FMockStaticState(), FMockDynamicState{0});
+		Manager.Game_EndTick(100);
+
+		// Tick 2
+		Manager.Game_BeginTick();
+		Handle1 = {};
+		Manager.Game_EndTick(200);
+
+		// Tick 3
+		Manager.Game_BeginTick();
+		Handle2 = {};
+		Manager.Game_EndTick(300);
+
+		// Tick 4
+		Manager.Game_BeginTick();
+		Manager.Game_EndTick(400);
+	}
+
+	FStateStreamManagerImpl& Manager = ManagerImpl;
+	FMockStateStreamImpl& Stream = MockStateStreamImpl;
+
+	SECTION("Rend not crossing ticks")
+    {
+		Manager.Render_Update(10); // We start consuming Tick 1
+		CHECK_EQUAL(Stream.Instances.Num(), 2);
+
+		Manager.Render_Update(100); // Consume rest of Tick 1
+		CHECK_EQUAL(Stream.Instances.Num(), 2);
+
+		Manager.Render_Update(110); // Moving into Tick 2
+		CHECK_EQUAL(Stream.Instances.Num(), 1);
+
+		Manager.Render_Update(200); // Consume rest of Tick 2
+		CHECK_EQUAL(Stream.Instances.Num(), 1);
+
+		Manager.Render_Update(210); // Moving into Tick 3
+		CHECK_EQUAL(Stream.Instances.Num(), 0);
+
+		Manager.Render_Update(300); // Consume rest of Tick 3
+		CHECK_EQUAL(Stream.Instances.Num(), 0);
+
+		Manager.Render_Update(400); // Consume rest of Tick 4
+		CHECK_EQUAL(Stream.Instances.Num(), 0);
+	}
+}
+
+TEST_CASE("StateStream::GarbageCollect", "[StateStream]")
+{
+	FMockStateStreamImpl MockStateStreamImpl;
+	FStateStreamManagerImpl ManagerImpl;
+	ManagerImpl.Render_Register(MockStateStreamImpl, false);
+
+	{
+		IStateStreamManager& Manager = ManagerImpl;
+		IMockStateStream& Stream = Manager.Game_Get<IMockStateStream>();
+
+		// Tick 1
+		Manager.Game_BeginTick();
+		FMockHandle Handle1 = Stream.Game_CreateInstance(FMockStaticState(), FMockDynamicState{0});
+		FMockHandle Handle2 = Stream.Game_CreateInstance(FMockStaticState(), FMockDynamicState{0});
+		Manager.Game_EndTick(100);
+		CHECK_EQUAL(MockStateStreamImpl.GetUsedInstancesCount(), 2);
+		CHECK_EQUAL(MockStateStreamImpl.GetUsedDynamicstatesCount(), 2);
+
+		// Tick 2
+		Manager.Game_BeginTick();
+		Handle1 = {};
+		Handle2 = {};
+		Manager.Game_EndTick(200);
+
+		// Tick 3
+		Manager.Game_BeginTick();
+		Manager.Game_EndTick(300);
+		CHECK_EQUAL(MockStateStreamImpl.GetUsedInstancesCount(), 2);
+		CHECK_EQUAL(MockStateStreamImpl.GetUsedDynamicstatesCount(), 2);
+	}
+
+	FStateStreamManagerImpl& Manager = ManagerImpl;
+	FMockStateStreamImpl& Stream = MockStateStreamImpl;
+	{
+		Manager.Render_Update(10);
+		Manager.Render_GarbageCollect();
+		CHECK_EQUAL(MockStateStreamImpl.GetUsedDynamicstatesCount(), 2);
+
+		Manager.Render_Update(100);
+		Manager.Render_GarbageCollect();
+		CHECK_EQUAL(MockStateStreamImpl.GetUsedDynamicstatesCount(), 2);
+
+		Manager.Render_Update(110);
+		Manager.Render_GarbageCollect();
+		CHECK_EQUAL(MockStateStreamImpl.GetUsedDynamicstatesCount(), 2);
+
+		Manager.Render_Update(200);
+		Manager.Render_GarbageCollect();
+
+		Manager.Render_Update(210);
+		Manager.Render_GarbageCollect();
+		CHECK_EQUAL(Stream.GetUsedInstancesCount(), 2);
+		CHECK_EQUAL(MockStateStreamImpl.GetUsedDynamicstatesCount(), 2);
+
+		Manager.Render_Update(300);
+		Manager.Render_GarbageCollect();
+		CHECK_EQUAL(Stream.GetUsedInstancesCount(), 0);
+		CHECK_EQUAL(MockStateStreamImpl.GetUsedDynamicstatesCount(), 0);
+	}
+}
+
+TEST_CASE("StateStream::InternalDependencies", "[StateStream]")
+{
+	FTransformStateStreamImpl TransformStreamImpl;
+
+
+	FStateStreamManagerImpl ManagerImpl;
+	ManagerImpl.Render_Register(TransformStreamImpl, false);
+
+
+	{
+		IStateStreamManager& Manager = ManagerImpl;
+		ITransformStateStream& TransformStream = Manager.Game_Get<ITransformStateStream>();
+
+		FTransformHandle TransformHandle;
+
+		// Tick 1
+		Manager.Game_BeginTick();
+		{
+			FTransformDynamicState Ds;
+			Ds.SetParent(TransformStream.Game_CreateInstance(FTransformStaticState(), FTransformDynamicState{}));
+			TransformHandle = TransformStream.Game_CreateInstance(FTransformStaticState(), Ds);
+		}
+		Manager.Game_EndTick(100);
+		CHECK_EQUAL(TransformStreamImpl.GetUsedInstancesCount(), 2);
+		CHECK_EQUAL(TransformStreamImpl.GetUsedDynamicstatesCount(), 2);
+
+		// Tick 2
+		Manager.Game_BeginTick();
+		{
+			FTransformDynamicState Ds;
+			Ds.SetParent(TransformStream.Game_CreateInstance(FTransformStaticState(), FTransformDynamicState{}));
+			TransformHandle.Update(Ds);
+		}
+		Manager.Game_EndTick(200);
+
+		// Tick 3
+		Manager.Game_BeginTick();
+		TransformHandle = {};
+		Manager.Game_EndTick(300);
+
+		Manager.Game_Exit();
+	}
+
+
+	FStateStreamManagerImpl& Manager = ManagerImpl;
+	FTransformStateStreamImpl& TransformStream = TransformStreamImpl;
+
+	SECTION("Rend not crossing ticks")
+    {
+		// Frame 1 consuming tick 1
+		Manager.Render_Update(100);
+
+		// Frame 2 consuming tick 2
+		Manager.Render_Update(200);
+
+		// Frame 3 consuming tick 3
+		Manager.Render_Update(300);
+	}
+
+	SECTION("Rend crossing ticks")
+    {
+		// Frame 1 consuming tick 1, 2 and 3
+		Manager.Render_Update(300);
+	}
+}
+
+TEST_CASE("StateStream::ExternalDependencies", "[StateStream]")
+{
+	FTransformStateStreamImpl TransformStreamImpl;
+	FMockStateStreamImpl MockStreamImpl;
+
+	FStateStreamManagerImpl ManagerImpl;
+	ManagerImpl.Render_Register(TransformStreamImpl, false);
+	ManagerImpl.Render_Register(MockStreamImpl, false);
+
+
+	{
+		IStateStreamManager& Manager = ManagerImpl;
+		ITransformStateStream& TransformStream = Manager.Game_Get<ITransformStateStream>();
+		IMockStateStream& MockStream = Manager.Game_Get<IMockStateStream>();
+
+		FMockHandle MockHandle;
+
+		// Tick 1
+		Manager.Game_BeginTick();
+		{
+			FMockDynamicState Ds;
+			Ds.SetTransform(TransformStream.Game_CreateInstance(FTransformStaticState(), FTransformDynamicState{}));
+			MockHandle = MockStream.Game_CreateInstance(FMockStaticState(), Ds);
+		}
+		Manager.Game_EndTick(100);
+		CHECK_EQUAL(TransformStreamImpl.GetUsedInstancesCount(), 1);
+		CHECK_EQUAL(TransformStreamImpl.GetUsedDynamicstatesCount(), 1);
+		CHECK_EQUAL(MockStreamImpl.GetUsedInstancesCount(), 1);
+		CHECK_EQUAL(MockStreamImpl.GetUsedDynamicstatesCount(), 1);
+
+		// Tick 2
+		Manager.Game_BeginTick();
+		{
+			FMockDynamicState Ds;
+			Ds.SetTransform(TransformStream.Game_CreateInstance(FTransformStaticState(), FTransformDynamicState{}));
+			MockHandle.Update(Ds);
+		}
+		Manager.Game_EndTick(200);
+
+		// Tick 3
+		Manager.Game_BeginTick();
+		MockHandle = {};
+		Manager.Game_EndTick(300);
+
+		Manager.Game_Exit();
+	}
+
+
+	FStateStreamManagerImpl& Manager = ManagerImpl;
+	FTransformStateStreamImpl& TransformStream = TransformStreamImpl;
+	FMockStateStreamImpl& MockStream = MockStreamImpl;
+
+	SECTION("Rend not crossing ticks")
+    {
+		// Frame 1 consuming tick 1
+		Manager.Render_Update(100);
+		CHECK_EQUAL(MockStream.Instances.Num(), 1);
+		//CHECK_EQUAL(MockStream.Instances[0]->Value, 110);
+
+		// Frame 2 consuming tick 2
+		Manager.Render_Update(200);
+
+		// Frame 3 consuming tick 3
+		Manager.Render_Update(300);
+	}
+
+	SECTION("Rend crossing ticks")
+    {
+		// Frame 1 consuming tick 1, 2 and 3
+		Manager.Render_Update(300);
+	}
+}
+
+TEST_CASE("StateStream::Interpolation", "[StateStream]")
+{
+	FTransformStateStreamImpl TransformStreamImpl;
+	FStateStreamManagerImpl ManagerImpl;
+	ManagerImpl.Render_Register(TransformStreamImpl, false);
+
+	{
+		IStateStreamManager& Manager = ManagerImpl;
+		ITransformStateStream& Stream = Manager.Game_Get<ITransformStateStream>();
+
+		// Tick 1
+		Manager.Game_BeginTick();
+		FTransformHandle Handle = Stream.Game_CreateInstance(FTransformStaticState(), FTransformDynamicState());
+		Manager.Game_EndTick(0);
+
+		// Tick 2
+		Manager.Game_BeginTick();
+		FTransformDynamicState newState;
+		newState.SetLocalTransform(FTransform(FQuat::Identity, FVector(100, 100, 100)));
+		Handle.Update(newState);
+		Manager.Game_EndTick(100);
+
+		// Tick 3
+		Manager.Game_BeginTick();
+		Handle = {};
+		Manager.Game_EndTick(200);
+	}
+
+	FStateStreamManagerImpl& Manager = ManagerImpl;
+	FTransformStateStreamImpl& Stream = TransformStreamImpl;
+
+	SECTION("Interpolation of translation")
+    {
+		Manager.Render_Update(10); // We start consuming Tick 2
+		FTransformObject* Object = (FTransformObject*)Stream.Render_GetUserData(1u);
+		CHECK_EQUAL((FMath::Abs(Object->GetInfo().WorldTransform.GetTranslation().X - 10.0f) < 0.001f), true);
+	}
+}
diff --git a/Engine/Source/Programs/Shared/EpicGames.UHT/Exporters/CodeGen/UhtHeaderCodeGeneratorHFile.cs b/Engine/Source/Programs/Shared/EpicGames.UHT/Exporters/CodeGen/UhtHeaderCodeGeneratorHFile.cs
--- a/Engine/Source/Programs/Shared/EpicGames.UHT/Exporters/CodeGen/UhtHeaderCodeGeneratorHFile.cs
+++ b/Engine/Source/Programs/Shared/EpicGames.UHT/Exporters/CodeGen/UhtHeaderCodeGeneratorHFile.cs
@@ -418,6 +418,165 @@ namespace EpicGames.UHT.Exporters.CodeGen
 						}
 						builder.Append("); \\\r\n");
 					}
+
+					if (scriptStruct.MetaData.ContainsKey(UhtNames.StateStreamStaticState) || scriptStruct.MetaData.ContainsKey(UhtNames.StateStreamDynamicState))
+					{
+						builder.Append("\tenum Fields \\\r\n");
+						builder.Append("\t{ \\\r\n");
+						int enumValue = 1;
+						foreach (UhtProperty property in scriptStruct.Properties)
+						{
+							builder.Append("\t\tField_").Append(property.SourceName).Append(" = ").Append(enumValue).Append(", \\\r\n");
+							enumValue = enumValue << 1;
+						}
+						builder.Append("\t}; \\\r\n");
+
+						builder.Append("\tuint32 ModifiedFlags = 0; \\\r\n");
+
+						builder.Append("\t").Append(scriptStruct.SourceName).Append("() = default; \\\r\n");
+
+						builder.Append("\t").Append(scriptStruct.SourceName).Append("(FStateStreamCopyContext& Context, const ").Append(scriptStruct.SourceName).Append("& Other) \\\r\n");
+						bool isFirst = true;
+
+						builder.Append(isFirst ? "\t:\t" : "\t,\t").Append("ModifiedFlags(Other.ModifiedFlags) \\\r\n");
+						isFirst = false;
+
+						bool hasHandles = false;
+
+						foreach (UhtProperty property in scriptStruct.Properties)
+						{
+							builder.Append(isFirst ? "\t:\t" : "\t,\t").Append(property.SourceName).Append("(");
+							if (property is UhtStructProperty structProperty)
+							{
+								if (structProperty.ScriptStruct.MetaData.ContainsKey(UhtNames.StateStreamHandle))
+								{
+									hasHandles = true;
+									builder.Append("Context, ");
+								}
+							}
+							builder.Append("Other.").Append(property.SourceName);
+							builder.Append(") \\\r\n");
+							isFirst = false;
+						}
+
+						builder.Append("\t{} \\\r\n");
+
+						foreach (UhtProperty property in scriptStruct.Properties)
+						{
+							string propName = property.SourceName;
+							bool isBoolWithB = (property is UhtBoolProperty && propName.StartsWith("b", StringComparison.Ordinal));
+							var cleanPropName = isBoolWithB ? propName[1..] : propName;
+							builder.Append("\tinline bool ").Append(cleanPropName).Append("Modified() const { return !!(ModifiedFlags & Field_"). Append(propName).Append("); } \\\r\n");
+
+							builder.Append("\tinline void Set").Append(cleanPropName).Append("(const ");
+							property.AppendText(builder, UhtPropertyTextType.Generic);
+							builder.Append("& In) { ").Append(propName).Append(" = In; ModifiedFlags |= Field_").Append(propName).Append("; } \\\r\n");
+
+							if (isBoolWithB)
+							{
+								builder.Append("\tinline bool Get").Append(cleanPropName).Append("() const { return ").Append(propName).Append("; } \\\r\n");
+							}
+							else
+							{
+								builder.Append("\tinline const ");
+								property.AppendText(builder, UhtPropertyTextType.Generic);
+								builder.Append("& Get").Append(propName).Append("() const { return ").Append(propName).Append("; } \\\r\n");
+							}
+						}
+
+						int dependencyCount = 0;
+						builder.Append("\tuint32 GetDependencies(FStateStreamHandle** OutHandles, uint32 Capacity) \\\r\n");
+						builder.Append("\t{ \\\r\n");
+						if (hasHandles)
+						{
+							foreach (UhtProperty property in scriptStruct.Properties)
+							{
+								if (property is UhtStructProperty structProperty)
+								{
+									if (structProperty.ScriptStruct.MetaData.ContainsKey(UhtNames.StateStreamHandle))
+									{
+										builder.Append("\t\tOutHandles[").Append(dependencyCount++).Append("] = &").Append(property.SourceName).Append("; \\\r\n");
+									}
+								}
+							}
+						}
+						builder.Append("\t\treturn ").Append(dependencyCount).Append("; \\\r\n");
+						builder.Append("\t} \\\r\n");
+
+						builder.Append("\tvoid Apply(FStateStreamCopyContext& Context, const ").Append(scriptStruct.SourceName).Append("& Other) \\\r\n");
+						builder.Append("\t{ \\\r\n");
+						foreach (UhtProperty property in scriptStruct.Properties)
+						{
+							string propName = property.SourceName;
+							bool isBoolWithB = (property is UhtBoolProperty && propName.StartsWith("b", StringComparison.Ordinal));
+							var cleanPropName = isBoolWithB ? propName[1..] : propName;
+
+							builder.Append("\t\tif (Other.").Append(cleanPropName).Append("Modified()");
+							builder.Append(" && !StateStreamEquals(").Append(propName).Append(", Other.").Append(propName);
+							builder.Append(")) \\\r\n");
+							builder.Append("\t\t{ \\\r\n");
+
+							bool assignHandled = false;
+							if (property is UhtStructProperty structProperty)
+							{
+								if (structProperty.ScriptStruct.MetaData.ContainsKey(UhtNames.StateStreamHandle))
+								{
+									builder.Append("\t\t\t").Append(propName).Append(".Apply(Context, Other.").Append(propName).Append("); \\\r\n");
+									assignHandled = true;
+								}
+							}
+							if (!assignHandled)
+							{
+								builder.Append("\t\t\t").Append(propName).Append(" = Other.").Append(propName).Append("; \\\r\n");
+							}
+							builder.Append("\t\t\tModifiedFlags |= Field_").Append(propName).Append("; \\\r\n");
+							builder.Append("\t\t} \\\r\n");
+						}
+						builder.Append("\t} \\\r\n");
+
+						builder.Append("\tvoid Interpolate(FStateStreamInterpolateContext& Context, const ").Append(scriptStruct.SourceName).Append("& From, const ").Append(scriptStruct.SourceName).Append("& To) \\\r\n");
+						builder.Append("\t{ \\\r\n");
+						builder.Append("\t\tModifiedFlags = 0; \\\r\n");
+						foreach (UhtProperty property in scriptStruct.Properties)
+						{
+							string propName = property.SourceName;
+							bool isBool = property is UhtBoolProperty;
+							bool isBoolWithB = isBool && propName.StartsWith("b", StringComparison.Ordinal);
+							var cleanPropName = isBoolWithB ? propName[1..] : propName;
+
+							builder.Append("\t\tif (To.").Append(cleanPropName).Append("Modified()) \\\r\n");
+							builder.Append("\t\t{ \\\r\n");
+							builder.Append("\t\t\tModifiedFlags |= Field_").Append(propName).Append("; \\\r\n");
+
+							bool assignHandled = false;
+							if (property is UhtStructProperty structProperty)
+							{
+								if (structProperty.ScriptStruct.MetaData.ContainsKey(UhtNames.StateStreamHandle))
+								{
+									builder.Append("\t\t\t").Append(propName).Append(".Apply(Context, To.").Append(propName).Append("); \\\r\n");
+									assignHandled = true;
+								}
+							}
+							if (!assignHandled)
+							{
+								if (isBool)
+								{
+									builder.Append("\t\t\t").Append(propName).Append(" = To.").Append(propName).Append(";\\\r\n");
+								}
+								else
+								{
+									builder.Append("\t\t\tStateStreamInterpolate(Context, ").Append(propName).Append(", From.").Append(propName).Append(", To.").Append(propName).Append("); \\\r\n");
+								}
+							}
+							builder.Append("\t\t} \\\r\n");
+						}
+						builder.Append("\t} \\\r\n");
+					}
+					else if (scriptStruct.MetaData.ContainsKey(UhtNames.StateStreamHandle))
+					{
+						ReadOnlySpan<char> shortName = scriptStruct.SourceName.AsSpan(1, scriptStruct.SourceName.Length - 7);
+						builder.Append("\tDECLARE_STATESTREAM_HANDLE(").Append(shortName).Append(") \\\r\n");
+					}
 				}
 
 				// Forward declare the StaticStruct specialization in the header
diff --git a/Engine/Source/Programs/Shared/EpicGames.UHT/Specifiers/UhtDefaultSpecifiers.cs b/Engine/Source/Programs/Shared/EpicGames.UHT/Specifiers/UhtDefaultSpecifiers.cs
--- a/Engine/Source/Programs/Shared/EpicGames.UHT/Specifiers/UhtDefaultSpecifiers.cs
+++ b/Engine/Source/Programs/Shared/EpicGames.UHT/Specifiers/UhtDefaultSpecifiers.cs
@@ -168,6 +168,24 @@ namespace EpicGames.UHT.Parsers
 		{
 			SetMetaData(specifierContext, "SparseClassDataType", value);
 		}
+
+		[UhtSpecifier(Extends = UhtTableNames.Default, ValueType = UhtSpecifierValueType.None, When = UhtSpecifierWhen.Immediate)]
+		private static void StateStreamStaticStateSpecifier(UhtSpecifierContext specifierContext)
+		{
+			SetMetaData(specifierContext, "StateStreamStaticState", true);
+		}
+
+		[UhtSpecifier(Extends = UhtTableNames.Default, ValueType = UhtSpecifierValueType.None, When = UhtSpecifierWhen.Immediate)]
+		private static void StateStreamDynamicStateSpecifier(UhtSpecifierContext specifierContext)
+		{
+			SetMetaData(specifierContext, "StateStreamDynamicState", true);
+		}
+
+		[UhtSpecifier(Extends = UhtTableNames.Default, ValueType = UhtSpecifierValueType.None, When = UhtSpecifierWhen.Immediate)]
+		private static void StateStreamHandleSpecifier(UhtSpecifierContext specifierContext)
+		{
+			SetMetaData(specifierContext, "StateStreamHandle", true);
+		}
 		#endregion
 
 		#region Validators
diff --git a/Engine/Source/Programs/Shared/EpicGames.UHT/Utils/UhtNames.cs b/Engine/Source/Programs/Shared/EpicGames.UHT/Utils/UhtNames.cs
--- a/Engine/Source/Programs/Shared/EpicGames.UHT/Utils/UhtNames.cs
+++ b/Engine/Source/Programs/Shared/EpicGames.UHT/Utils/UhtNames.cs
@@ -20,6 +20,9 @@ namespace EpicGames.UHT.Utils
 		public const string BlueprintSetter = "BlueprintSetter";
 		public const string CallInEditor = "CallInEditor";
 		public const string BlueprintType = "BlueprintType";
+		public const string StateStreamStaticState = "StateStreamStaticState";
+		public const string StateStreamDynamicState = "StateStreamDynamicState";
+		public const string StateStreamHandle = "StateStreamHandle";
 		public const string CannotImplementInterfaceInBlueprint = "CannotImplementInterfaceInBlueprint";
 		public const string Category = "Category";
 		public const string ClassGroupNames = "ClassGroupNames";
diff --git a/Engine/Source/Runtime/StateStream/Private/StateStreamHandle.cpp b/Engine/Source/Runtime/StateStream/Private/StateStreamHandle.cpp
new file mode 100644
--- /dev/null
+++ b/Engine/Source/Runtime/StateStream/Private/StateStreamHandle.cpp
@@ -0,0 +1,127 @@
+// Copyright Epic Games, Inc. All Rights Reserved.
+
+#include "StateStreamHandle.h"
+
+////////////////////////////////////////////////////////////////////////////////////////////////////
+
+FStateStreamHandle::FStateStreamHandle(const FStateStreamHandle& Other)
+:	Owner(Other.Owner)
+,	IsInternal(Other.IsInternal)
+,	Id(Other.Id)
+{
+	if (Owner && !IsInternal)
+	{
+		Owner->Game_AddRef(Id);
+	}
+}
+
+FStateStreamHandle::FStateStreamHandle(FStateStreamHandle&& Other)
+{
+	Owner = Other.Owner;
+	IsInternal = Other.IsInternal;
+	Id = Other.Id;
+	Other.Owner = nullptr;
+	Other.Id = 0;
+}
+
+FStateStreamHandle& FStateStreamHandle::operator=(FStateStreamHandle&& Other)
+{
+	if (Owner && !IsInternal)
+	{
+		Owner->Game_Release(Id);
+	}
+	Owner = Other.Owner;
+	IsInternal = Other.IsInternal;
+	Id = Other.Id;
+	Other.Owner = nullptr;
+	Other.Id = 0;
+	return *this;
+}
+
+FStateStreamHandle& FStateStreamHandle::operator=(const FStateStreamHandle& Other)
+{
+	if (Owner == Other.Owner && Id == Other.Id)
+	{
+		return *this;
+	}
+	if (Owner && !IsInternal)
+	{
+		Owner->Game_Release(Id);
+	}
+	Owner = Other.Owner;
+	IsInternal = Other.IsInternal;
+	Id = Other.Id;
+	if (Owner && !IsInternal)
+		Owner->Game_AddRef(Id);
+	return *this;
+}
+
+FStateStreamHandle::~FStateStreamHandle()
+{
+	if (Owner && !IsInternal)
+	{
+		Owner->Game_Release(Id);
+	}
+}
+
+FStateStreamHandle::FStateStreamHandle(IStateStreamHandleOwner& O, uint32 I)
+:	Owner(&O)
+,	Id(I)
+{
+}
+
+FStateStreamHandle::FStateStreamHandle(FStateStreamCopyContext& Context, const FStateStreamHandle& Other)
+:	Owner(Other.Owner)
+,	IsInternal(Context.IsInternal)
+,	Id(Other.Id)
+{
+	if (Owner && !IsInternal)
+	{
+		Owner->Game_AddRef(Id);
+	}
+}
+
+void FStateStreamHandle::Apply(FStateStreamCopyContext& Context, const FStateStreamHandle& Other)
+{
+	if (!IsInternal)
+	{
+		check(!Other.IsInternal); // If handle is not internal we should be on game side and the other must also not be internal
+		if (Other.Owner)
+		{
+			Other.Owner->Game_AddRef(Other.Id);
+		}
+		if (Owner)
+		{
+			Owner->Game_Release(Id);
+		}
+	}
+
+	Owner = Other.Owner;
+	Id = Other.Id;
+}
+
+bool FStateStreamHandle::operator==(const FStateStreamHandle& Other) const
+{
+	return Id == Other.Id && Owner == Other.Owner;
+}
+
+void* FStateStreamHandle::Render_GetUserData() const
+{
+	if (Owner && Id)
+	{
+		check(IsInternal);
+		return Owner->Render_GetUserData(Id);
+	}
+	return nullptr;
+}
+
+void FStateStreamHandle::MakeInternal()
+{
+	if(Owner && !IsInternal)
+	{
+		Owner->Game_Release(Id);
+	}
+	IsInternal = true;
+}
+
+////////////////////////////////////////////////////////////////////////////////////////////////////
diff --git a/Engine/Source/Runtime/StateStream/Private/StateStreamManagerImpl.cpp b/Engine/Source/Runtime/StateStream/Private/StateStreamManagerImpl.cpp
new file mode 100644
--- /dev/null
+++ b/Engine/Source/Runtime/StateStream/Private/StateStreamManagerImpl.cpp
@@ -0,0 +1,137 @@
+// Copyright Epic Games, Inc. All Rights Reserved.
+
+#include "StateStreamManagerImpl.h"
+#include "StateStreamDebugRenderer.h"
+
+////////////////////////////////////////////////////////////////////////////////////////////////////
+
+void FStateStreamManagerImpl::Game_BeginTick()
+{
+	check(!bIsInTick);
+	check (!bGameExited);
+	bIsInTick = true;
+	for (const StateStreamRec& Rec : StateStreams)
+	{
+		Rec.Stream->Game_BeginTick();
+	}
+}
+
+void FStateStreamManagerImpl::Game_EndTick(double AbsoluteTime)
+{
+	check(bIsInTick);
+	bIsInTick = false;
+	for (const StateStreamRec& Rec : StateStreams)
+	{
+		Rec.Stream->Game_EndTick(AbsoluteTime);
+	}
+}
+
+void FStateStreamManagerImpl::Game_Exit()
+{
+	check(!bIsInTick);
+	for (const StateStreamRec& Rec : StateStreams)
+	{
+		Rec.Stream->Game_Exit();
+	}
+	bGameExited = true;
+}
+
+bool FStateStreamManagerImpl::Game_IsInTick()
+{
+	return bIsInTick;
+}
+
+void* FStateStreamManagerImpl::Game_GetStreamPointer(uint32 Id)
+{
+	check(Id < uint32(StateStreamsLookup.Num()));
+	return StateStreamsLookup[Id]->Game_GetVoidPointer();
+}
+
+void FStateStreamManagerImpl::Render_Register(IStateStream& Stream, bool TakeOwnership)
+{
+	StateStreamRec Rec { &Stream, TakeOwnership };
+	StateStreams.Add(Rec);
+	
+	uint32 Id = Stream.GetId();
+	if (uint32(StateStreamsLookup.Num()) <= Id)
+	{
+		StateStreamsLookup.SetNum(Id+1);
+	}
+	StateStreamsLookup[Id] = &Stream;
+}
+
+void FStateStreamManagerImpl::Render_RegisterDependency(uint32 FromId, uint32 ToId)
+{
+	// TODO
+}
+
+void FStateStreamManagerImpl::Render_RegisterDependency(IStateStream& From, IStateStream& To)
+{
+	Render_RegisterDependency(From.GetId(), To.GetId());
+}
+
+void FStateStreamManagerImpl::Render_Update(double AbsoluteTime)
+{
+	for (const StateStreamRec& Rec : StateStreams)
+	{
+		Rec.Stream->Render_Update(AbsoluteTime);
+	}
+
+	for (const StateStreamRec& Rec : StateStreams)
+	{
+		Rec.Stream->Render_PostUpdate();
+	}
+}
+
+void FStateStreamManagerImpl::Render_Exit()
+{
+	GarbageCollectTask.Wait();
+
+	for (const StateStreamRec& Rec : StateStreams)
+	{
+		Rec.Stream->Render_Exit();
+	}
+	bRenderExited = true;
+}
+
+void FStateStreamManagerImpl::Render_GarbageCollect(bool AsTask)
+{
+	GarbageCollectTask.Wait();
+
+	auto Func = [this]()
+		{
+			for (const StateStreamRec& Rec : StateStreams)
+			{
+				Rec.Stream->Render_GarbageCollect();
+			}
+		};
+	if (AsTask)
+	{
+		GarbageCollectTask = UE::Tasks::Launch(UE_SOURCE_LOCATION, MoveTemp(Func));
+	}
+	else
+	{
+		Func();
+	}
+}
+
+void FStateStreamManagerImpl::Game_DebugRender(IStateStreamDebugRenderer& Renderer)
+{
+	for (const StateStreamRec& Rec : StateStreams)
+	{
+		Rec.Stream->DebugRender(Renderer);
+	}
+}
+
+FStateStreamManagerImpl::~FStateStreamManagerImpl()
+{
+	for (const StateStreamRec& Rec : StateStreams)
+	{
+		if (Rec.Owned)
+		{
+			delete Rec.Stream;
+		}
+	}
+}
+
+////////////////////////////////////////////////////////////////////////////////////////////////////
diff --git a/Engine/Source/Runtime/StateStream/Private/TransformStateStreamImpl.cpp b/Engine/Source/Runtime/StateStream/Private/TransformStateStreamImpl.cpp
new file mode 100644
--- /dev/null
+++ b/Engine/Source/Runtime/StateStream/Private/TransformStateStreamImpl.cpp
@@ -0,0 +1,138 @@
+// Copyright Epic Games, Inc. All Rights Reserved.
+
+#include "TransformStateStreamImpl.h"
+
+FTransformObject::~FTransformObject()
+{
+	if (Parent)
+	{
+		Parent->RemoveListener(this);
+		Parent = nullptr;
+	}
+	check(!First);
+}
+
+FTransformObject::Info FTransformObject::GetInfo()
+{
+	// TODO: Race condition safe
+	if (bDirty)
+	{
+		WorldTransform = LocalState->GetLocalTransform();
+		bVisible = LocalState->GetVisible();
+
+		if (Parent)
+		{
+			Info ParentInfo = Parent->GetInfo();
+			WorldTransform *= ParentInfo.WorldTransform;
+			bVisible |= ParentInfo.bVisible;
+		}
+
+		bDirty = false;
+	}
+
+	return { WorldTransform, LocalState->GetBoneTransforms(), bVisible };
+}
+
+void FTransformObject::AddListener(FTransformObjectListener* Listener)
+{
+	if (!First)
+	{
+		First = Listener;
+	}
+	else
+	{
+		Listener->Next = First;
+		First->Prev = Listener;
+		First = Listener;
+	}
+}
+
+void FTransformObject::RemoveListener(FTransformObjectListener* Listener)
+{
+	if (Listener->Prev)
+	{
+		Listener->Prev->Next = Listener->Next;
+	}
+	else
+	{
+		First = Listener->Next;
+	}
+
+	if (Listener->Next)
+	{
+		Listener->Next->Prev = Listener->Prev;
+	}
+
+	Listener->Next = nullptr;
+	Listener->Prev = nullptr;
+}
+
+void FTransformObject::CallListeners()
+{
+	for (FTransformObjectListener* It = First; It; It = It->Next)
+	{
+		It->OnTransformObjectDirty();
+	}
+}
+
+void FTransformObject::OnTransformObjectDirty()
+{
+	bDirty = true;
+	CallListeners();
+}
+
+void FTransformStateStreamImpl::SetParent(FTransformObject* Object, const FTransformDynamicState& Ds)
+{
+	const FTransformHandle& ParentHandle = Ds.GetParent();
+	if (!ParentHandle.IsValid() || Object->Parent)
+	{
+		return;
+	}
+
+	FTransformObject*& ParentUserData = Render_GetUserData(ParentHandle);
+	FTransformObject* ParentObject = ParentUserData;
+	if (!ParentObject)
+	{
+		ParentObject = new FTransformObject();
+		ParentObject->AddRef();
+		ParentUserData = ParentObject;
+
+		SetParent(ParentObject, Render_GetDynamicState(ParentHandle));
+	}
+
+	ParentObject->AddListener(Object);
+	Object->Parent = ParentObject;
+}
+
+void FTransformStateStreamImpl::Render_OnCreate(const FTransformStaticState& Ss, const FTransformDynamicState& Ds, FTransformObject*& UserData, bool IsDestroyedInSameFrame)
+{
+	FTransformObject* Object = UserData;
+	if (!Object)
+	{
+		Object = new FTransformObject();
+		Object->AddRef();
+		Object->LocalState = &Ds;
+		UserData = Object;
+	}
+
+	SetParent(Object, Ds);
+}
+
+void FTransformStateStreamImpl::Render_OnUpdate(const FTransformStaticState& Ss, const FTransformDynamicState& Ds, FTransformObject*& UserData)
+{
+		FTransformObject* Object = UserData;
+		if (Ds.LocalTransformModified() || Ds.VisibleModified() || Ds.BoneTransformsModified())
+		{
+			Object->OnTransformObjectDirty();
+		}
+
+		if (Ds.ParentModified())
+		{
+			SetParent(Object, Ds);
+		}
+}
+
+void FTransformStateStreamImpl::Render_OnDestroy(const FTransformStaticState& Ss, const FTransformDynamicState& Ds, FTransformObject*& UserData)
+{
+	UserData->Release();
+}
diff --git a/Engine/Source/Runtime/StateStream/Public/GenericStateStream.h b/Engine/Source/Runtime/StateStream/Public/GenericStateStream.h
new file mode 100644
--- /dev/null
+++ b/Engine/Source/Runtime/StateStream/Public/GenericStateStream.h
@@ -0,0 +1,632 @@
+// Copyright Epic Games, Inc. All Rights Reserved.
+
+#pragma once
+
+#include "HAL/CriticalSection.h"
+#include "Misc/Optional.h"
+#include "Misc/ScopeRWLock.h"
+#include "StateStream.h"
+#include "StateStreamDebugRenderer.h"
+#include "StateStreamDefinitions.h"
+#include "StateStreamStore.h"
+
+////////////////////////////////////////////////////////////////////////////////////////////////////
+// Settings struct - Can be inherited to set settings on statestream
+
+template<typename TInterfaceType, typename TUserDataType = void>
+struct TStateStreamSettings
+{
+	using InterfaceType = TInterfaceType;
+	using UserDataType = TUserDataType;
+
+	static inline constexpr bool SkipCreatingDeletes = false;
+};
+
+
+////////////////////////////////////////////////////////////////////////////////////////////////////
+// TStateStream is a generic implementation of IStateStream that contains all the boiler plate code
+// related to ticks, interpolation, etc. Should be the default goto implementation.
+// Inherit this class for each type of state stream
+// Provide a subclass of TStateStreamSettings as template parameter
+
+template<typename Settings>
+class TStateStream : public IStateStream, public Settings::InterfaceType, public IStateStreamHandleOwner
+{
+public:
+	using InterfaceType = Settings::InterfaceType;
+	using FHandle = typename InterfaceType::Handle;
+	using FStaticState = typename InterfaceType::StaticState;
+	using FDynamicState = typename InterfaceType::DynamicState;
+	using FUserDataType = typename Settings::UserDataType;
+
+	// IStateStreamOwner (used by state stream handles on game side)
+	virtual FHandle Game_CreateInstance(const FStaticState& Ss, const FDynamicState& Ds) override final;
+	virtual void Game_AddRef(uint32 HandleId) override final;
+	virtual void Game_Release(uint32 HandleId) override final;
+	virtual void Game_Update(uint32 HandleId, const void* Ds) override final;
+	virtual void* Game_Edit(uint32 HandleId) override final;
+	virtual void* Render_GetUserData(uint32 HandleId) override final;
+	
+	// IStateStream (used by StateStreamManagerImpl)
+	virtual void Game_BeginTick() override;
+	virtual void Game_EndTick(StateStreamTime AbsoluteTime) override;
+	virtual void Game_Exit() override;
+	virtual void* Game_GetVoidPointer() override;
+	virtual uint32 GetId() override;
+	virtual void Render_Update(StateStreamTime AbsoluteTime) override;
+	virtual void Render_PostUpdate() override;
+	virtual void Render_Exit() override;
+	virtual void Render_GarbageCollect() override;
+
+	virtual const TCHAR* GetDebugName() override;
+	virtual void DebugRender(IStateStreamDebugRenderer& Renderer) override;
+
+
+	FUserDataType*& Render_GetUserData(const FHandle& Handle);
+	const FDynamicState& Render_GetDynamicState(const FHandle& Handle);
+
+
+	// Specialize to do custom things
+	virtual void Render_OnCreate(const FStaticState& Ss, const FDynamicState& Ds, FUserDataType*& UserData, bool IsDestroyedInSameFrame) {}
+	virtual void Render_OnUpdate(const FStaticState& Ss, const FDynamicState& Ds, FUserDataType*& UserData) {}
+	virtual void Render_OnDestroy(const FStaticState& Ss, const FDynamicState& Ds, FUserDataType*& UserData) {}
+
+	// ... or use template specialization to avoid virtual calls
+	void Render_OnCreateInline(const FStaticState& Ss, const FDynamicState& Ds, FUserDataType*& UserData, bool IsDestroyedInSameFrame) { Render_OnCreate(Ss, Ds, UserData, IsDestroyedInSameFrame); }
+	void Render_OnUpdateInline(const FStaticState& Ss, const FDynamicState& Ds, FUserDataType*& UserData) { Render_OnUpdate(Ss, Ds, UserData); }
+	void Render_OnDestroyInline(const FStaticState& Ss, const FDynamicState& Ds, FUserDataType*& UserData) { Render_OnDestroy(Ss, Ds, UserData); }
+	
+	// For unit tests
+	uint32 GetUsedInstancesCount() const { return Instances.GetUsedCount(); }
+	uint32 GetUsedDynamicstatesCount() const { return DynamicStates.GetUsedCount(); }
+
+	~TStateStream();
+
+private:
+	struct FTick;
+
+	void ApplyChanges(const FTick& Tick, StateStreamTime Time, uint32 PrevTickIndex, StateStreamTime PrevTime, const TBitArray<>& ModifiedInstances);
+
+	template<typename T>
+	void MakeInternal(T& State);
+
+	// Information about instance.
+	struct FInstance
+	{
+		FInstance() = default;
+		FInstance(const FStaticState& Ss, uint32 Rc, uint32 Ct) : StaticState(Ss), RefCount(Rc), CreateTick(Ct) {}
+		FStaticState StaticState;
+		uint32 RefCount = 0;
+		uint32 CreateTick = 0;
+		uint32 DeleteTick = ~0u;
+		TOptional<FDynamicState> RendDynamicState;
+		FUserDataType* UserData = nullptr;
+	};
+
+	TStateStreamStore<FInstance> Instances;
+
+	// Tick produced by game side.
+	struct FTick
+	{
+		TArray<uint32> DynamicStates; // Contains index of all instances existing in tick (some might have been destroyed but index not reused)
+		TBitArray<> ModifiedInstances; // Contains bits saying which of the instances that have been modified in this tick.
+		FTick* PrevTick = nullptr; // Tick with earlier time
+		FTick* NextTick = nullptr; // Tick with newer time
+		StateStreamTime Time = 0; // Time Tick finished
+		uint32 Index = 0; // Index of tick (created by TickCounter)
+	};
+
+	FRWLock CurrentTickLock;
+	FTick* CurrentTick = nullptr; // Tick being worked on in game
+	FTick* OldestAvailableTick = nullptr; // Newest finished tick available to rendering
+	FTick* NewestAvailableTick = nullptr; // Newest finished tick available to rendering
+
+	TStateStreamStore<FDynamicState> DynamicStates; // Store for dynamic states.
+
+	uint32 TickCounter = 1;
+
+	FTick* RendTick = nullptr; // Last used tick for rendering.
+	StateStreamTime RendTime = 0; // Last used time for rendering
+
+	TArray<FInstance*> DeferredDestroys;
+};
+
+
+////////////////////////////////////////////////////////////////////////////////////////////////////
+// Implementation
+
+template<typename Settings>
+typename TStateStream<Settings>::FHandle TStateStream<Settings>::Game_CreateInstance(const FStaticState& Ss, const FDynamicState& Ds)
+{
+	check(CurrentTick);
+
+	uint32 InstanceIndex = Instances.Emplace(Ss, 1, CurrentTick->Index);
+
+	uint32 DynamicStateIndex = DynamicStates.Add(Ds);
+
+	CurrentTickLock.WriteLock();
+
+	if (uint32(CurrentTick->DynamicStates.Num()) <= InstanceIndex)
+	{
+		CurrentTick->DynamicStates.SetNum(InstanceIndex + 1);
+		CurrentTick->ModifiedInstances.SetNum(InstanceIndex + 1, false);
+	}
+	CurrentTick->ModifiedInstances[InstanceIndex] = true;
+	CurrentTick->DynamicStates[InstanceIndex] = DynamicStateIndex;
+
+	CurrentTickLock.WriteUnlock();
+
+	return { *this, InstanceIndex + 1 };
+}
+
+template<typename Settings>
+void TStateStream<Settings>::Game_AddRef(uint32 HandleId)
+{
+	check(CurrentTick);
+	check(HandleId != 0);
+	uint32 InstanceIndex = HandleId - 1;
+	FInstance& Instance = Instances[InstanceIndex];
+	check(Instance.RefCount);
+	++Instance.RefCount;
+}
+
+template<typename Settings>
+void TStateStream<Settings>::Game_Release(uint32 HandleId)
+{
+	check(CurrentTick);
+	check(HandleId != 0);
+	uint32 InstanceIndex = HandleId - 1;
+	FInstance& Instance = Instances[InstanceIndex];
+	check(Instance.RefCount);
+	if (--Instance.RefCount)
+	{
+		return;
+	}
+
+	Instance.DeleteTick = CurrentTick->Index;
+	MakeInternal(Instance.StaticState);
+
+	CurrentTickLock.ReadLock();
+
+	CurrentTick->ModifiedInstances[InstanceIndex] = true;
+	uint32 DsIndex = CurrentTick->DynamicStates[InstanceIndex];
+
+	CurrentTickLock.ReadUnlock();
+
+	MakeInternal(DynamicStates[DsIndex]);
+}
+
+template<typename Settings>
+void TStateStream<Settings>::Game_Update(uint32 HandleId, const void* Ds)
+{
+	FDynamicState* DsState = static_cast<FDynamicState*>(Game_Edit(HandleId));
+	FStateStreamCopyContext Context;
+	Context.IsInternal = 0;
+	DsState->Apply(Context, *static_cast<const FDynamicState*>(Ds));
+}
+
+template<typename Settings>
+void* TStateStream<Settings>::Game_Edit(uint32 HandleId)
+{
+	check(HandleId != 0);
+	uint32 InstanceIndex = HandleId - 1;
+
+	check(CurrentTick);
+	check(InstanceIndex < uint32(CurrentTick->DynamicStates.Num()));
+
+	CurrentTickLock.ReadLock();
+
+	FDynamicState* DsState;
+
+	if (CurrentTick->ModifiedInstances[InstanceIndex])
+	{
+		uint32 DsIndex = CurrentTick->DynamicStates[InstanceIndex];
+		CurrentTickLock.ReadUnlock();
+		DsState = &DynamicStates[DsIndex];
+	}
+	else
+	{
+		uint32& IndexRef = CurrentTick->DynamicStates[InstanceIndex];
+		uint32 OldIndex = IndexRef;
+		CurrentTick->ModifiedInstances[InstanceIndex] = true;
+		uint32 DsIndex;
+		void* DsPtr = DynamicStates.AddUninitialized(DsIndex);
+		IndexRef = DsIndex;
+		CurrentTickLock.ReadUnlock();
+
+		FDynamicState& OldDs = DynamicStates[OldIndex];
+		DsState = new (DsPtr) FDynamicState(OldDs);
+		MakeInternal(OldDs);
+	}
+	return DsState;
+}
+
+
+template<typename Settings>
+void* TStateStream<Settings>::Render_GetUserData(uint32 HandleId)
+{
+	check(HandleId);
+	uint32 InstanceIndex = HandleId - 1;
+	return Instances[InstanceIndex].UserData;
+}
+
+template<typename Settings>
+void TStateStream<Settings>::Game_BeginTick()
+{
+	check(!CurrentTick);
+	CurrentTick = new FTick();
+	CurrentTick->Index = TickCounter++;
+	if (NewestAvailableTick)
+	{
+		CurrentTick->DynamicStates = NewestAvailableTick->DynamicStates;
+		CurrentTick->ModifiedInstances.SetNum(CurrentTick->DynamicStates.Num(), false);
+	}
+}
+
+template<typename Settings>
+void TStateStream<Settings>::Game_EndTick(StateStreamTime AbsoluteTime)
+{
+	check(CurrentTick);
+
+	CurrentTick->Time = AbsoluteTime;
+	CurrentTick->PrevTick = NewestAvailableTick;
+	if (NewestAvailableTick)
+	{
+		check(NewestAvailableTick->Time <= AbsoluteTime);
+		NewestAvailableTick->NextTick = CurrentTick;
+	}
+	else
+	{
+		OldestAvailableTick = CurrentTick;
+	}
+		
+	NewestAvailableTick = CurrentTick;
+	CurrentTick = nullptr;
+}
+
+template<typename Settings>
+void TStateStream<Settings>::Game_Exit()
+{
+	check(!CurrentTick);
+	/*
+	for (TConstSetBitIterator<> It(NewestAvailableTick->Modif); It; ++It)
+	{
+
+	}
+	*/
+}
+
+template<typename Settings>
+void* TStateStream<Settings>::Game_GetVoidPointer()
+{
+	return static_cast<InterfaceType*>(this);
+}
+
+template<typename Settings>
+uint32 TStateStream<Settings>::GetId()
+{
+	return Settings::InterfaceType::Id;
+}
+
+template<typename Settings>
+void TStateStream<Settings>::Render_Update(StateStreamTime AbsoluteTime)
+{
+	if (!NewestAvailableTick || AbsoluteTime == RendTime)
+		return;
+	check(AbsoluteTime > RendTime); // Only play forward for now
+
+	uint32 PrevTick = 0;
+	bool IsFirstTick = false;
+	if (RendTick)
+	{
+		PrevTick = RendTick->Index;
+	}
+	else
+	{
+		RendTick = OldestAvailableTick;
+		IsFirstTick = true;
+	}
+
+	StateStreamTime PrevTime = RendTime;
+	RendTime = AbsoluteTime;
+
+	// We are still inside the same tick
+	if (RendTime <= RendTick->Time)
+	{
+		// Just interpolate or apply RendTick
+		ApplyChanges(*RendTick, RendTime, PrevTick, PrevTime, RendTick->ModifiedInstances);
+		return;
+	}
+		
+	// We were at exact end of last handled tick.. move into next
+	if (PrevTime == RendTick->Time)
+	{
+		FTick* NewRendTick = RendTick->NextTick;
+
+		if (!NewRendTick) // We've caught up with game.. set RendTime back to PrevTime and return
+		{
+			RendTime = PrevTime;
+			return;
+		}
+
+		if (!IsFirstTick)
+			RendTick = NewRendTick;
+
+		// We don't need to include RendTick Modifications.
+		if (RendTime <= RendTick->Time)
+		{
+			// Just interpolate or apply RendTick
+			ApplyChanges(*RendTick, RendTime, PrevTick, PrevTime, RendTick->ModifiedInstances);
+			return;
+		}
+	}
+
+	// We are overlapping between two or more ticks
+	TBitArray<> ModifiedInstances(RendTick->ModifiedInstances);
+	while (RendTick->Time < RendTime)
+	{
+		if (!RendTick->NextTick)
+		{
+			RendTime = RendTick->Time;
+			break;
+		}
+		RendTick = RendTick->NextTick;
+		ModifiedInstances.CombineWithBitwiseOR(RendTick->ModifiedInstances, EBitwiseOperatorFlags::MaxSize);
+	}
+
+	ApplyChanges(*RendTick, RendTime, PrevTick, PrevTime, ModifiedInstances);
+}
+
+template<typename Settings>
+void TStateStream<Settings>::Render_PostUpdate()
+{
+	for (FInstance* Instance : DeferredDestroys)
+	{
+		Render_OnDestroyInline(Instance->StaticState, *Instance->RendDynamicState, Instance->UserData);
+		Instance->RendDynamicState.Reset();
+	}
+	DeferredDestroys.SetNum(0);
+}
+
+template<typename Settings>
+void TStateStream<Settings>::Render_Exit()
+{
+	if (NewestAvailableTick)
+	{
+		StateStreamTime maxTime = (StateStreamTime)std::numeric_limits<StateStreamTime>::max();
+		Render_Update(maxTime);
+		//Render_GarbageCollect();
+	}
+}
+
+template<typename Settings>
+void TStateStream<Settings>::Render_GarbageCollect()
+{
+	if (!OldestAvailableTick)
+	{
+		return;
+	}
+
+	FTick* RendTickUsed = RendTick;
+	if (RendTick && RendTick->Time != RendTime && RendTick->PrevTick)
+		RendTickUsed = RendTick->PrevTick;
+	
+	FTick* Tick = OldestAvailableTick;
+	while (Tick != RendTickUsed)
+	{
+		FTick* Next = Tick->NextTick;
+		check (Next);
+
+		// We can remove all DynamicStates that are different in Next tick since we know this is the last tick using the state
+		for (TConstSetBitIterator<> It(Next->ModifiedInstances); It; ++It)
+		{
+			uint32 InstanceIndex = It.GetIndex();
+			if (InstanceIndex >= uint32(Tick->DynamicStates.Num()))
+			{
+				continue;
+			}
+
+			// If instance was created in next tick we ignore this. If instance is deleted in tick we handle it further down
+			FInstance& Instance = Instances[InstanceIndex];
+			if (Instance.DeleteTick == Tick->Index || Instance.CreateTick == Next->Index)
+			{
+				continue;
+			}
+
+			uint32 DynamicStateIndex = Tick->DynamicStates[InstanceIndex];
+			if (DynamicStateIndex == Next->DynamicStates[InstanceIndex])
+			{
+				continue;
+			}
+
+			DynamicStates.Remove(DynamicStateIndex);
+		}
+
+		// Remove all instances who was deleted in Tick
+		for (TConstSetBitIterator<> It(Tick->ModifiedInstances); It; ++It)
+		{
+			uint32 InstanceIndex = It.GetIndex();
+			FInstance& Instance = Instances[InstanceIndex];
+
+			if (Instance.DeleteTick != Tick->Index)
+			{
+				continue;
+			}
+
+			DynamicStates.Remove(Tick->DynamicStates[InstanceIndex]);
+			Instance.RendDynamicState.Reset();
+			Instances.Remove(InstanceIndex);
+		}
+
+		delete Tick;
+		Tick = Next;
+	}
+	OldestAvailableTick = Tick;
+
+}
+
+template<typename Settings>
+const TCHAR* TStateStream<Settings>::GetDebugName()
+{
+	return FHandle::GetDebugName();
+}
+
+template<typename Settings>
+void TStateStream<Settings>::DebugRender(IStateStreamDebugRenderer& Renderer)
+{
+	TStringBuilder<1024> DebugLine;
+	uint32 ModifiedCount = 0;
+	if (FTick* Tick = NewestAvailableTick)
+	{
+		ModifiedCount = Tick->ModifiedInstances.CountSetBits();
+	}
+	DebugLine.Appendf(TEXT("%s   Num: %u  Changed: %u"), FHandle::GetDebugName(), Instances.GetUsedCount(), ModifiedCount);
+	Renderer.DrawText(*DebugLine);
+}
+
+template<typename Settings>
+typename TStateStream<Settings>::FUserDataType*& TStateStream<Settings>::Render_GetUserData(const FHandle& Handle)
+{
+	uint32 Id = Handle.GetId();
+	check(Id);
+	uint32 InstanceIndex = Id - 1;
+	return Instances[InstanceIndex].UserData;
+}
+
+template<typename Settings>
+const typename TStateStream<Settings>::FDynamicState& TStateStream<Settings>::Render_GetDynamicState(const FHandle& Handle)
+{
+	uint32 Id = Handle.GetId();
+	check(Id);
+	uint32 InstanceIndex = Id - 1;
+	return *Instances[InstanceIndex].RendDynamicState;
+}
+
+template<typename Settings>
+void TStateStream<Settings>::ApplyChanges(const FTick& Tick, StateStreamTime Time, uint32 PrevTickIndex, StateStreamTime PrevTime, const TBitArray<>& ModifiedInstances)
+{
+	// TODO: This will do lots of allocations
+	struct FCreateInfo { FInstance* Instance; bool IsDestroyInSameFrame; };
+	TArray<FCreateInfo> Creates;
+	TArray<FInstance*> Updates;
+	TArray<FInstance*> Destroys;
+
+	for (TConstSetBitIterator<> It(ModifiedInstances); It; ++It)
+	{
+		uint32 InstanceIndex = It.GetIndex();
+		FInstance& Instance = Instances[InstanceIndex];
+
+		if (Instance.DeleteTick <= PrevTickIndex) // Already deleted
+		{
+			continue;
+		}
+
+		FStaticState& Ss = Instance.StaticState;
+
+		bool IsCreate = Instance.CreateTick > PrevTickIndex && Instance.CreateTick <= Tick.Index;
+		bool IsDestroy = Instance.DeleteTick > PrevTickIndex && Instance.DeleteTick <= Tick.Index;
+
+		if (IsCreate)
+		{
+			if (Settings::SkipCreatingDeletes && IsDestroy)
+			{
+				continue;
+			}
+
+			check(!Instance.RendDynamicState.IsSet());
+			uint32 DynamicStateIndex = Tick.DynamicStates[InstanceIndex];
+			FStateStreamCopyContext Context;
+			Instance.RendDynamicState.Emplace(Context, DynamicStates[DynamicStateIndex]);
+		}
+		else if (Instance.CreateTick == Tick.Index) // Still in the tick it was created, no interpolation possible
+		{
+			continue;
+		}
+
+		FDynamicState& Ds = *Instance.RendDynamicState;
+
+		if (Tick.Time == Time || !Tick.PrevTick)
+		{
+			if (!IsCreate)
+			{
+				FStateStreamCopyContext Context;
+				Ds.Apply(Context, DynamicStates[Tick.DynamicStates[InstanceIndex]]);
+			}
+		}
+		else
+		{
+			const FTick& PrevTick = *Tick.PrevTick;
+			uint32 FromIndex = PrevTick.DynamicStates[InstanceIndex];
+			uint32 ToIndex = Tick.DynamicStates[InstanceIndex];
+			
+			StateStreamTime DeltaTime = Tick.Time - PrevTick.Time;
+			StateStreamTime TimeInTo = Time - PrevTick.Time;
+			double Factor = double(TimeInTo) / DeltaTime;
+
+			FDynamicState& From = DynamicStates[FromIndex];
+			FDynamicState& To = DynamicStates[ToIndex];
+
+			FStateStreamInterpolateContext Context;
+			Context.Factor = Factor;
+			Ds.Interpolate(Context, From, To);
+		}
+
+		if (IsCreate)
+		{
+			Creates.Add({&Instance, IsDestroy});
+		}
+
+		if (IsDestroy)
+		{
+			if (IsCreate)
+			{
+				DeferredDestroys.Add(&Instance);
+			}
+			else
+			{
+				Destroys.Add(&Instance);
+			}
+			continue;
+		}
+
+		if (!IsCreate)
+		{
+			Updates.Add(&Instance);
+		}
+	}
+
+	for (FCreateInfo& Info : Creates)
+	{
+		FInstance& Instance = *Info.Instance;
+		Render_OnCreateInline(Instance.StaticState, *Instance.RendDynamicState, Instance.UserData, Info.IsDestroyInSameFrame);
+	}
+
+	for (FInstance* Instance : Updates)
+	{
+		Render_OnUpdateInline(Instance->StaticState, *Instance->RendDynamicState, Instance->UserData);
+	}
+
+	for (FInstance* Instance : Destroys)
+	{
+		Render_OnDestroyInline(Instance->StaticState, *Instance->RendDynamicState, Instance->UserData);
+		Instance->RendDynamicState.Reset();
+	}
+}
+
+template<typename Settings>
+template<typename T>
+void TStateStream<Settings>::MakeInternal(T& State)
+{
+	FStateStreamHandle* Deps[256];
+	for (uint32 I=0, E=State.GetDependencies(Deps, UE_ARRAY_COUNT(Deps)); I!=E; ++I)
+	{
+		Deps[I]->MakeInternal();
+	}
+}
+
+template<typename Settings>
+TStateStream<Settings>::~TStateStream()
+{
+}
+
+////////////////////////////////////////////////////////////////////////////////////////////////////
diff --git a/Engine/Source/Runtime/StateStream/Public/Pow2ChunkedArray.h b/Engine/Source/Runtime/StateStream/Public/Pow2ChunkedArray.h
new file mode 100644
--- /dev/null
+++ b/Engine/Source/Runtime/StateStream/Public/Pow2ChunkedArray.h
@@ -0,0 +1,175 @@
+// Copyright Epic Games, Inc. All Rights Reserved.
+
+#pragma once
+
+#include "HAL/PlatformMath.h"
+#include "HAL/UnrealMemory.h"
+
+
+////////////////////////////////////////////////////////////////////////////////////////////////////
+// FPow2ChunkedArray
+//
+// A chunk based array where each chunk is twice as large as the previous one.
+// Adding and accessing elements are O(1)
+// Add is threadsafe and keeping references is safe since it never reallocates
+
+template<typename T, uint32 MinSize = 16, uint32 MaxSize = 16777216>
+class FPow2ChunkedArray
+{
+public:
+	enum : uint32 { SkipCount = FGenericPlatformMath::CeilLogTwo(MinSize) };
+	enum : uint32 { BucketCount = FGenericPlatformMath::CeilLogTwo(MaxSize) - SkipCount + 1 };
+
+	inline FPow2ChunkedArray();
+	inline ~FPow2ChunkedArray();
+
+	inline T& Add(const T& Value, uint32& OutIndex);
+	inline T& Add(const T& Value);
+	inline T& Add(T&& Value);
+
+	inline const T& operator[](uint32 Index) const;
+	inline T& operator[](uint32 Index);
+
+	inline uint32 Num() const;
+
+	inline const T& GetElementAt(uint32 Index) const;
+
+	inline void* AddUninitialized(uint32& OutIndex);
+
+	inline uint32 GetBucketIndex(uint32 Index) const;
+	inline uint32 GetBucketStart(uint32 BucketIndex) const;
+	inline uint32 GetBucketSize(uint32 BucketIndex) const;
+
+private:
+	std::atomic<uint32> Size;
+	volatile int64 Buckets[BucketCount];
+};
+
+
+////////////////////////////////////////////////////////////////////////////////////////////////////
+// Implementation
+
+template<typename T, uint32 MinSize, uint32 MaxSize>
+FPow2ChunkedArray<T, MinSize, MaxSize>::FPow2ChunkedArray()
+{
+	FMemory::Memset((int64*)Buckets, 0, sizeof(Buckets));
+}
+
+template<typename T, uint32 MinSize, uint32 MaxSize>
+FPow2ChunkedArray<T, MinSize, MaxSize>::~FPow2ChunkedArray()
+{
+	uint32 Left = Size;
+	for (uint32 BucketIndex=0; Left; ++BucketIndex)
+	{
+		T* Elements = reinterpret_cast<T*>(Buckets[BucketIndex]);
+		uint32 LeftInBucket = FPlatformMath::Min(Left, GetBucketSize(BucketIndex));
+		if (!std::is_trivially_destructible_v<T>)
+		{
+			for (uint32 I=0, E=LeftInBucket; I!=E; ++I)
+			{
+				Elements[I].~T();
+			}
+		}
+		Left -= LeftInBucket;
+		FMemory::Free(Elements);
+	}
+}
+
+template<typename T, uint32 MinSize, uint32 MaxSize>
+T& FPow2ChunkedArray<T, MinSize, MaxSize>::Add(const T& Value, uint32& OutIndex)
+{
+	return *new (AddUninitialized(OutIndex)) T(Value);
+}
+
+template<typename T, uint32 MinSize, uint32 MaxSize>
+T& FPow2ChunkedArray<T, MinSize, MaxSize>::Add(const T& Value)
+{
+	uint32 OutIndex;
+	return *new (AddUninitialized(OutIndex)) T(Value);
+}
+
+template<typename T, uint32 MinSize, uint32 MaxSize>
+T& FPow2ChunkedArray<T, MinSize, MaxSize>::Add(T&& Value)
+{
+	uint32 OutIndex;
+	return *new (AddUninitialized(OutIndex)) T(std::move(Value));
+}
+
+template<typename T, uint32 MinSize, uint32 MaxSize>
+const T& FPow2ChunkedArray<T, MinSize, MaxSize>::operator[](uint32 Index) const
+{
+	return GetElementAt(Index);
+}
+
+template<typename T, uint32 MinSize, uint32 MaxSize>
+T& FPow2ChunkedArray<T, MinSize, MaxSize>::operator[](uint32 Index)
+{
+	return const_cast<T&>(GetElementAt(Index));
+}
+
+template<typename T, uint32 MinSize, uint32 MaxSize>
+uint32 FPow2ChunkedArray<T, MinSize, MaxSize>::Num() const
+{
+	return Size;
+}
+
+template<typename T, uint32 MinSize, uint32 MaxSize>
+const T& FPow2ChunkedArray<T, MinSize, MaxSize>::GetElementAt(uint32 Index) const
+{
+	check(Index < Size);
+	uint32 BucketIndex = GetBucketIndex(Index);
+	uint32 BucketStart = GetBucketStart(BucketIndex);
+	uint32 BucketOffset = Index - BucketStart;
+	return reinterpret_cast<const T*>(Buckets[BucketIndex])[BucketOffset];
+}
+
+template<typename T, uint32 MinSize, uint32 MaxSize>
+void* FPow2ChunkedArray<T, MinSize, MaxSize>::AddUninitialized(uint32& OutIndex)
+{
+	uint32 Index = Size++;
+	uint32 BucketIndex = GetBucketIndex(Index);
+		
+	volatile int64& BucketRef = Buckets[BucketIndex];
+		
+	int64 BucketRes = FPlatformAtomics::AtomicRead(&BucketRef);
+	if (!BucketRes)
+	{
+		void* NewPtr = FMemory::Malloc(GetBucketSize(BucketIndex) * sizeof(T), alignof(T));
+		BucketRes = FPlatformAtomics::InterlockedCompareExchange(&BucketRef, int64(NewPtr), 0);
+		if (BucketRes)
+		{
+			FMemory::Free(NewPtr);
+		}
+		else
+		{
+			BucketRes = int64(NewPtr);
+		}
+	}
+
+	uint32 BucketStart = GetBucketStart(BucketIndex);
+	uint32 BucketOffset = Index - BucketStart;
+
+	OutIndex = Index;
+
+	return reinterpret_cast<T*>(BucketRes) + BucketOffset;
+}
+
+template<typename T, uint32 MinSize, uint32 MaxSize>
+uint32 FPow2ChunkedArray<T, MinSize, MaxSize>::GetBucketIndex(uint32 Index) const
+{
+	return uint32(FPlatformMath::FloorLog2(Index / 16 + 1));
+}
+
+template<typename T, uint32 MinSize, uint32 MaxSize>
+uint32 FPow2ChunkedArray<T, MinSize, MaxSize>::GetBucketStart(uint32 BucketIndex) const
+{
+	return BucketIndex ? 16 * ((1 << BucketIndex) - 1) : 0;
+}
+
+template<typename T, uint32 MinSize, uint32 MaxSize>
+uint32 FPow2ChunkedArray<T, MinSize, MaxSize>::GetBucketSize(uint32 BucketIndex) const
+{
+	return FPlatformMath::Max(MinSize, 1u << (BucketIndex + SkipCount));
+}
+
+////////////////////////////////////////////////////////////////////////////////////////////////////
diff --git a/Engine/Source/Runtime/StateStream/Public/StateStream.h b/Engine/Source/Runtime/StateStream/Public/StateStream.h
new file mode 100644
--- /dev/null
+++ b/Engine/Source/Runtime/StateStream/Public/StateStream.h
@@ -0,0 +1,34 @@
+// Copyright Epic Games, Inc. All Rights Reserved.
+
+#pragma once
+
+#include "StateStreamDefinitions.h"
+
+class IStateStreamDebugRenderer;
+
+////////////////////////////////////////////////////////////////////////////////////////////////////
+// Interface used by StateStreamManagerImpl. This is RT only. Should not be visible/used by GT
+// Documentation in StateStreamManagerImpl
+
+class IStateStream
+{
+public:
+	virtual void Game_BeginTick() = 0;
+	virtual void Game_EndTick(StateStreamTime AbsoluteTime) = 0;
+	virtual void Game_Exit() = 0;
+	virtual void* Game_GetVoidPointer() = 0;
+
+	virtual void Render_Update(StateStreamTime AbsoluteTime) = 0;
+	virtual void Render_PostUpdate() = 0;
+	virtual void Render_Exit() = 0;
+	virtual void Render_GarbageCollect() = 0;
+
+	virtual uint32 GetId() = 0;
+
+	virtual const TCHAR* GetDebugName() { return TEXT("Unknown"); }
+	virtual void DebugRender(IStateStreamDebugRenderer& Renderer) {}
+
+	virtual ~IStateStream() {}
+};
+
+////////////////////////////////////////////////////////////////////////////////////////////////////
diff --git a/Engine/Source/Runtime/StateStream/Public/StateStreamDebugRenderer.h b/Engine/Source/Runtime/StateStream/Public/StateStreamDebugRenderer.h
new file mode 100644
--- /dev/null
+++ b/Engine/Source/Runtime/StateStream/Public/StateStreamDebugRenderer.h
@@ -0,0 +1,15 @@
+// Copyright Epic Games, Inc. All Rights Reserved.
+
+#pragma once
+
+#include "Containers/StringFwd.h"
+
+////////////////////////////////////////////////////////////////////////////////////////////////////
+
+class IStateStreamDebugRenderer
+{
+public:
+	virtual void DrawText(const FStringView& Text) = 0;
+};
+
+////////////////////////////////////////////////////////////////////////////////////////////////////
diff --git a/Engine/Source/Runtime/StateStream/Public/StateStreamDefinitions.h b/Engine/Source/Runtime/StateStream/Public/StateStreamDefinitions.h
new file mode 100644
--- /dev/null
+++ b/Engine/Source/Runtime/StateStream/Public/StateStreamDefinitions.h
@@ -0,0 +1,53 @@
+// Copyright Epic Games, Inc. All Rights Reserved.
+
+#pragma once
+
+#include "StateStreamHandle.h"
+
+////////////////////////////////////////////////////////////////////////////////////////////////////
+
+#define DECLARE_STATESTREAM(Type) \
+	using StaticState = F##Type##StaticState; \
+	using DynamicState = F##Type##DynamicState; \
+	using Handle = F##Type##Handle; \
+	static inline constexpr uint32 Id = Type##StateStreamId; \
+
+
+// TODO: A struct with multiple times (wall, game, etc)
+#if !defined(UE_STATESTREAM_TIME_TYPE)
+#define UE_STATESTREAM_TIME_TYPE double
+#endif
+
+using StateStreamTime = UE_STATESTREAM_TIME_TYPE;
+
+////////////////////////////////////////////////////////////////////////////////////////////////////
+// Context used for interpolation functions
+
+struct FStateStreamInterpolateContext : FStateStreamCopyContext
+{
+	double Factor = 0;
+};
+
+////////////////////////////////////////////////////////////////////////////////////////////////////
+// Generic interpolation functions. Used by code generation
+
+template<typename Type>
+inline void StateStreamInterpolate(FStateStreamInterpolateContext& Context, Type& Out, const Type& From, const Type& To)
+{
+	Out = To;
+}
+
+
+inline void StateStreamInterpolate(FStateStreamInterpolateContext& Context, uint32& Out, const uint32& From, const uint32& To)
+{
+	Out = From + uint32(double(To - From)*Context.Factor);
+}
+
+
+template<typename Type>
+inline bool StateStreamEquals(const Type& A, const Type& B)
+{
+	return A == B;
+}
+
+////////////////////////////////////////////////////////////////////////////////////////////////////
diff --git a/Engine/Source/Runtime/StateStream/Public/StateStreamHandle.h b/Engine/Source/Runtime/StateStream/Public/StateStreamHandle.h
new file mode 100644
--- /dev/null
+++ b/Engine/Source/Runtime/StateStream/Public/StateStreamHandle.h
@@ -0,0 +1,98 @@
+// Copyright Epic Games, Inc. All Rights Reserved.
+
+#pragma once
+
+#include "StateStreamHandle.generated.h"
+
+#define UE_API STATESTREAM_API
+
+class IStateStreamHandleOwner;
+struct FStateStreamCopyContext;
+
+////////////////////////////////////////////////////////////////////////////////////////////////////
+// Base type for state stream handles. This code is not supposed to be used directly.
+// To create your new state stream you will need a handle dedicated to that state stream. Example:
+// 
+//   USTRUCT(StateStreamHandle)
+//   struct FFooHandle : public FStateStreamHandle
+//   {
+//      GENERATED_USTRUCT_BODY()
+//   };
+//
+
+USTRUCT(NoExport)
+struct FStateStreamHandle
+{
+	UE_API void MakeInternal();
+
+	inline bool IsValid() const { return Id != 0; }
+	inline uint32 GetId() const { return Id; }
+
+	UE_API void* Render_GetUserData() const;
+
+protected:
+	FStateStreamHandle() = default;
+	UE_API FStateStreamHandle(const FStateStreamHandle& Other);
+	UE_API FStateStreamHandle(FStateStreamHandle&& Other);
+	UE_API FStateStreamHandle& operator=(FStateStreamHandle&& Other);
+	UE_API FStateStreamHandle& operator=(const FStateStreamHandle& Other);
+	UE_API ~FStateStreamHandle();
+
+	UE_API FStateStreamHandle(IStateStreamHandleOwner& O, uint32 I);
+	UE_API FStateStreamHandle(FStateStreamCopyContext& Context, const FStateStreamHandle& Other);
+	UE_API void Apply(FStateStreamCopyContext& Context, const FStateStreamHandle& Other);
+
+	UE_API bool operator==(const FStateStreamHandle& Other) const;
+
+	IStateStreamHandleOwner* Owner = nullptr;
+	uint32 IsInternal = 0;
+	uint32 Id = 0;
+};
+
+
+////////////////////////////////////////////////////////////////////////////////////////////////////
+// StateStream handle owner.
+
+class IStateStreamHandleOwner
+{
+public:
+	virtual void Game_AddRef(uint32 HandleId) = 0;
+	virtual void Game_Release(uint32 HandleId) = 0;
+	virtual void Game_Update(uint32 HandleId, const void* DynamicState) = 0;
+	virtual void* Game_Edit(uint32 HandleId) = 0;
+	virtual void* Render_GetUserData(uint32 HandleId) = 0;
+};
+
+
+////////////////////////////////////////////////////////////////////////////////////////////////////
+// Copy context used to copy state stream handles.
+
+struct FStateStreamCopyContext
+{
+	uint32 IsInternal = 1;
+};
+
+
+////////////////////////////////////////////////////////////////////////////////////////////////////
+// Macro added by UHT to USTRUCT(StateStreamHandle) types
+
+#define DECLARE_STATESTREAM_HANDLE(Type) \
+	F##Type##Handle() = default; \
+	F##Type##Handle(const F##Type##Handle& Other) : FStateStreamHandle(Other) {} \
+	F##Type##Handle(F##Type##Handle&& Other) : FStateStreamHandle(std::move(Other)) {} \
+	F##Type##Handle& operator=(F##Type##Handle&& Other) { FStateStreamHandle::operator=(std::move(Other)); return *this; } \
+	F##Type##Handle& operator=(const F##Type##Handle& Other) { FStateStreamHandle::operator=(Other); return *this; } \
+	~F##Type##Handle() = default; \
+	void Update(const F##Type##DynamicState& Ds) { if (Owner) Owner->Game_Update(Id, &Ds); } \
+	F##Type##DynamicState& Edit() { return *static_cast<F##Type##DynamicState*>(Owner->Game_Edit(Id)); } \
+	F##Type##Handle(IStateStreamHandleOwner& O, uint32 I) : FStateStreamHandle(O, I) {} \
+	F##Type##Handle(FStateStreamCopyContext& C, const F##Type##Handle& O) : FStateStreamHandle(C, O) {} \
+	void Apply(FStateStreamCopyContext& Context, const FStateStreamHandle& Other) { FStateStreamHandle::Apply(Context, Other); } \
+	bool operator==(const F##Type##Handle& Other) const { return FStateStreamHandle::operator==(Other); } \
+	static const TCHAR* GetDebugName() { return TEXT(#Type); }; \
+	using FStateStreamHandle::Apply; \
+	friend F##Type##DynamicState; \
+
+////////////////////////////////////////////////////////////////////////////////////////////////////
+
+#undef UE_API
diff --git a/Engine/Source/Runtime/StateStream/Public/StateStreamManager.h b/Engine/Source/Runtime/StateStream/Public/StateStreamManager.h
new file mode 100644
--- /dev/null
+++ b/Engine/Source/Runtime/StateStream/Public/StateStreamManager.h
@@ -0,0 +1,41 @@
+// Copyright Epic Games, Inc. All Rights Reserved.
+
+#pragma once
+
+#include "HAL/Platform.h"
+
+class IStateStreamDebugRenderer;
+
+////////////////////////////////////////////////////////////////////////////////////////////////////
+// StateStreamManager interface. This should be used from Game side
+
+class IStateStreamManager
+{
+public:
+	// Call from Game when a new tick is opened
+	// Note, no state stream handles can be created,updated,destroyed outside a Begin/End tick
+	virtual void Game_BeginTick() = 0;
+
+	// Close tick and make it available to render side.
+	// AbsoluteTime is the amount of time that Game consumed
+	virtual void Game_EndTick(double AbsoluteTime) = 0;
+
+	// Should be called when game is exiting.
+	virtual void Game_Exit() = 0;
+
+	// Returns true if game is inside an open tick
+	virtual bool Game_IsInTick() = 0;
+
+	// Functions to fetch StateStream interface (not IStateStream) game side.
+	template<typename T>
+	T& Game_Get() { return *static_cast<T*>(Game_GetStreamPointer(T::Id)); }
+	virtual void* Game_GetStreamPointer(uint32 Id) = 0;
+
+	// StateStream debug rendering
+	virtual void Game_DebugRender(IStateStreamDebugRenderer& Renderer) = 0;
+
+protected:
+	virtual ~IStateStreamManager() {}
+};
+
+////////////////////////////////////////////////////////////////////////////////////////////////////
diff --git a/Engine/Source/Runtime/StateStream/Public/StateStreamManagerImpl.h b/Engine/Source/Runtime/StateStream/Public/StateStreamManagerImpl.h
new file mode 100644
--- /dev/null
+++ b/Engine/Source/Runtime/StateStream/Public/StateStreamManagerImpl.h
@@ -0,0 +1,64 @@
+// Copyright Epic Games, Inc. All Rights Reserved.
+
+#pragma once
+
+#include "StateStream.h"
+#include "StateStreamManager.h"
+#include "Tasks/Task.h"
+
+#define UE_API STATESTREAM_API
+
+////////////////////////////////////////////////////////////////////////////////////////////////////
+// StateStreamManager implementation. 
+// This type should only be known render side.
+
+class FStateStreamManagerImpl : public IStateStreamManager
+{
+public:
+	// IStateStreamManager interface
+	UE_API virtual void Game_BeginTick() override;
+	UE_API virtual void Game_EndTick(double AbsoluteTime) override;
+	UE_API virtual void Game_Exit() override;
+	UE_API virtual bool Game_IsInTick() override;
+	UE_API virtual void* Game_GetStreamPointer(uint32 Id) override;
+
+	// Register new state streams into manager.
+	// TakeOwnership true means that manager will delete stream when shutting down
+	UE_API void Render_Register(IStateStream& Stream, bool TakeOwnership);
+
+	// Register dependency between statestreams. FromId will depend on ToId
+	UE_API void Render_RegisterDependency(uint32 FromId, uint32 ToId);
+	UE_API void Render_RegisterDependency(IStateStream& From, IStateStream& To);
+
+	// Called at the beginning of a render frame. AbsolutTime is the amount of time the render frame consumes
+	UE_API void Render_Update(double AbsoluteTime);
+
+	// Called before Render thread exits
+	UE_API void Render_Exit();
+
+	// Garbage collect
+	UE_API void Render_GarbageCollect(bool AsTask = false);
+
+	// Debug
+	UE_API virtual void Game_DebugRender(IStateStreamDebugRenderer& Renderer);
+
+	UE_API ~FStateStreamManagerImpl();
+
+private:
+	struct StateStreamRec
+	{
+		IStateStream* Stream;
+		bool Owned;
+	};
+	TArray<StateStreamRec> StateStreams;
+	TArray<IStateStream*> StateStreamsLookup;
+	bool bIsInTick = false;
+	bool bGameExited = false;
+	bool bRenderExited = false;
+
+	UE::Tasks::FTask GarbageCollectTask;
+};
+
+////////////////////////////////////////////////////////////////////////////////////////////////////
+
+#undef UE_API
diff --git a/Engine/Source/Runtime/StateStream/Public/StateStreamStore.h b/Engine/Source/Runtime/StateStream/Public/StateStreamStore.h
new file mode 100644
--- /dev/null
+++ b/Engine/Source/Runtime/StateStream/Public/StateStreamStore.h
@@ -0,0 +1,126 @@
+// Copyright Epic Games, Inc. All Rights Reserved.
+
+#pragma once
+
+#include "HAL/CriticalSection.h"
+#include "Pow2ChunkedArray.h"
+
+////////////////////////////////////////////////////////////////////////////////////////////////////
+// TStateStreamStore is a threadsafe, reference stable storage. User can hold on to pointers of
+// elements while other elements are added/removed.
+// Store is backed by a Pow2ChunkedArray and use a "free list" to be able to reuse removed elements
+
+template<typename T>
+class TStateStreamStore
+{
+public:
+	inline TStateStreamStore();
+	inline ~TStateStreamStore();
+
+	inline uint32 Add(const T& Value);
+	inline void* AddUninitialized(uint32& OutIndex);
+	inline void Remove(uint32 Index);
+
+	template <typename... ArgsType>
+	inline uint32 Emplace(ArgsType&&... Args);
+
+	inline T& operator[](uint32 Index);
+
+	inline uint32 GetUsedCount() const;
+
+private:
+	FRWLock Lock;
+	FPow2ChunkedArray<T> Array;
+	uint32 FirstFreeIndex = ~0u;
+};
+
+////////////////////////////////////////////////////////////////////////////////////////////////////
+// Implementation
+
+template<typename T>
+TStateStreamStore<T>::TStateStreamStore()
+{
+	static_assert(sizeof(T) >= sizeof(uint32));
+}
+
+template<typename T>
+TStateStreamStore<T>::~TStateStreamStore()
+{
+	if (!std::is_trivially_destructible_v<T>)
+	{
+		while (FirstFreeIndex != ~0u)
+		{
+			T& Element = Array[FirstFreeIndex];
+			FirstFreeIndex = *reinterpret_cast<uint32*>(&Element);
+			new (&Element) T();
+		}
+	}
+}
+
+template<typename T>
+uint32 TStateStreamStore<T>::Add(const T& Value)
+{
+	uint32 Index;
+	new (AddUninitialized(Index)) T(Value);
+	return Index;
+}
+
+template<typename T>
+void* TStateStreamStore<T>::AddUninitialized(uint32& OutIndex)
+{
+	Lock.WriteLock();
+	if (FirstFreeIndex == ~0u)
+	{
+		Lock.WriteUnlock();
+		return Array.AddUninitialized(OutIndex);
+	}
+
+	uint32 Index = FirstFreeIndex;
+	void* Element = &Array[Index];
+	FirstFreeIndex = *reinterpret_cast<uint32*>(Element);
+	Lock.WriteUnlock();
+
+	OutIndex = Index;
+	return &Array[Index];
+}
+
+template<typename T>
+void TStateStreamStore<T>::Remove(uint32 Index)
+{
+	T& Element = Array[Index];
+	Element.~T();
+	Lock.WriteLock();
+	*reinterpret_cast<uint32*>(&Element) = FirstFreeIndex;
+	FirstFreeIndex = Index;
+	Lock.WriteUnlock();
+}
+
+template<typename T>
+template <typename... ArgsType>
+uint32 TStateStreamStore<T>::Emplace(ArgsType&&... Args)
+{
+	uint32 OutIndex;
+	new (AddUninitialized(OutIndex)) T(Forward<ArgsType>(Args)...);
+	return OutIndex;
+}
+
+template<typename T>
+T& TStateStreamStore<T>::operator[](uint32 Index)
+{
+	return Array[Index];
+}
+
+template<typename T>
+uint32 TStateStreamStore<T>::GetUsedCount() const
+{
+	uint32 UsedCount = Array.Num();
+	uint32 Index = FirstFreeIndex;
+	while (Index != ~0u)
+	{
+		--UsedCount;
+		Index = *(uint32*)&Array[Index];
+	}
+	return UsedCount;
+}
+
+////////////////////////////////////////////////////////////////////////////////////////////////////
diff --git a/Engine/Source/Runtime/StateStream/Public/TransformStateStream.h b/Engine/Source/Runtime/StateStream/Public/TransformStateStream.h
new file mode 100644
--- /dev/null
+++ b/Engine/Source/Runtime/StateStream/Public/TransformStateStream.h
@@ -0,0 +1,71 @@
+// Copyright Epic Games, Inc. All Rights Reserved.
+
+#pragma once
+
+#include "TransformStateStreamHandle.h"
+#include "TransformStateStreamMath.h"
+#include "TransformStateStream.generated.h"
+
+
+////////////////////////////////////////////////////////////////////////////////////////////////////
+// Static state for transform instance. Can only be set upon creation
+
+USTRUCT(StateStreamStaticState)
+struct FTransformStaticState
+{
+	GENERATED_USTRUCT_BODY()
+};
+
+
+////////////////////////////////////////////////////////////////////////////////////////////////////
+// Dynamic state for transform instance. Can be updated inside ticks
+
+USTRUCT(StateStreamDynamicState)
+struct FTransformDynamicState
+{
+	GENERATED_USTRUCT_BODY()
+
+private: // Use accessors (SetX/GetX) instead of properties directly
+
+	// Transform relative Parent
+	UPROPERTY()
+	FTransform LocalTransform = FTransform::Identity;
+
+	UPROPERTY()
+	TArray<FTransform> BoneTransforms;
+
+	// Parent
+	UPROPERTY()
+	FTransformHandle Parent;
+
+	//UPROPERTY()
+	//uint8 LocationRule = 0;
+
+	//UPROPERTY()
+	//uint8 RotationRule = 0;
+
+	//UPROPERTY()
+	//uint8 ScaleRule = 0;
+
+	UPROPERTY()
+	bool bVisible = true;
+};
+
+
+////////////////////////////////////////////////////////////////////////////////////////////////////
+// Transform state stream id used for registering dependencies and find statestream
+
+inline constexpr uint32 TransformStateStreamId = 1;
+
+
+////////////////////////////////////////////////////////////////////////////////////////////////////
+// Interface for creating transform instances
+
+class ITransformStateStream
+{
+public:
+	DECLARE_STATESTREAM(Transform)
+	virtual FTransformHandle Game_CreateInstance(const FTransformStaticState& Ss, const FTransformDynamicState& Ds) = 0;
+};
+
+////////////////////////////////////////////////////////////////////////////////////////////////////
diff --git a/Engine/Source/Runtime/StateStream/Public/TransformStateStreamHandle.h b/Engine/Source/Runtime/StateStream/Public/TransformStateStreamHandle.h
new file mode 100644
--- /dev/null
+++ b/Engine/Source/Runtime/StateStream/Public/TransformStateStreamHandle.h
@@ -0,0 +1,19 @@
+// Copyright Epic Games, Inc. All Rights Reserved.
+
+#pragma once
+
+#include "StateStreamHandle.h"
+#include "TransformStateStreamHandle.generated.h"
+
+struct FTransformDynamicState;
+
+////////////////////////////////////////////////////////////////////////////////////////////////////
+// Handle for transform instance
+
+USTRUCT(StateStreamHandle)
+struct FTransformHandle : public FStateStreamHandle
+{
+	GENERATED_USTRUCT_BODY()
+};
+
+////////////////////////////////////////////////////////////////////////////////////////////////////
diff --git a/Engine/Source/Runtime/StateStream/Public/TransformStateStreamImpl.h b/Engine/Source/Runtime/StateStream/Public/TransformStateStreamImpl.h
new file mode 100644
--- /dev/null
+++ b/Engine/Source/Runtime/StateStream/Public/TransformStateStreamImpl.h
@@ -0,0 +1,75 @@
+// Copyright Epic Games, Inc. All Rights Reserved.
+
+#pragma once
+
+#include "GenericStateStream.h"
+#include "Templates/RefCounting.h"
+#include "TransformStateStream.h"
+
+#define UE_API STATESTREAM_API
+
+////////////////////////////////////////////////////////////////////////////////////////////////////
+// Listener for when transform object gets dirty.
+// Inherit FTransformObjectListener and register to FTransformObject
+
+class FTransformObjectListener
+{
+public:
+	virtual void OnTransformObjectDirty() = 0;
+
+private:
+	FTransformObjectListener* Prev = nullptr;
+	FTransformObjectListener* Next= nullptr;
+	friend class FTransformObject;
+};
+
+
+////////////////////////////////////////////////////////////////////////////////////////////////////
+// Transform object representing a instantiated transform.
+
+class FTransformObject : public TRefCountingMixin<FTransformObject>, public FTransformObjectListener
+{
+public:
+	struct Info
+	{
+		const FTransform& WorldTransform;
+		const TArray<FTransform>& BoneTransforms;
+		bool bVisible;
+	};
+
+	UE_API Info GetInfo();
+
+	UE_API void AddListener(FTransformObjectListener* Listener);
+	UE_API void RemoveListener(FTransformObjectListener* Listener);
+
+private:
+	virtual ~FTransformObject();
+	void CallListeners();
+	virtual void OnTransformObjectDirty() override final;
+
+	TRefCountPtr<FTransformObject> Parent;
+	FTransformObjectListener* First = 0;
+	const FTransformDynamicState* LocalState = nullptr;
+	FTransform WorldTransform;
+	bool bDirty = true;
+	bool bVisible = true;
+
+	friend class FTransformStateStreamImpl;
+	friend class TRefCountingMixin<FTransformObject>;
+};
+
+
+////////////////////////////////////////////////////////////////////////////////////////////////////
+
+class FTransformStateStreamImpl : public TStateStream<TStateStreamSettings<ITransformStateStream, FTransformObject>>
+{
+private:
+	UE_API void SetParent(FTransformObject* Object, const FTransformDynamicState& Ds);
+	UE_API virtual void Render_OnCreate(const FTransformStaticState& Ss, const FTransformDynamicState& Ds, FTransformObject*& UserData, bool IsDestroyedInSameFrame) override;
+	UE_API virtual void Render_OnUpdate(const FTransformStaticState& Ss, const FTransformDynamicState& Ds, FTransformObject*& UserData) override;
+	UE_API virtual void Render_OnDestroy(const FTransformStaticState& Ss, const FTransformDynamicState& Ds, FTransformObject*& UserData) override;
+};
+
+////////////////////////////////////////////////////////////////////////////////////////////////////
+
+#undef UE_API
diff --git a/Engine/Source/Runtime/StateStream/Public/TransformStateStreamMath.h b/Engine/Source/Runtime/StateStream/Public/TransformStateStreamMath.h
new file mode 100644
--- /dev/null
+++ b/Engine/Source/Runtime/StateStream/Public/TransformStateStreamMath.h
@@ -0,0 +1,38 @@
+// Copyright Epic Games, Inc. All Rights Reserved.
+
+#pragma once
+
+#include "Math/Transform.h"
+#include "StateStreamDefinitions.h"
+
+////////////////////////////////////////////////////////////////////////////////////////////////////
+// Special implementations for types used in state stream states
+// Used by generated code
+
+inline void StateStreamInterpolate(FStateStreamInterpolateContext& Context, FTransform& Out, const FTransform& From, const FTransform& To)
+{
+	Out.Blend(From, To, Context.Factor);
+}
+
+inline bool StateStreamEquals(const FTransform& A, const FTransform& B)
+{
+	return false;
+}
+
+inline void StateStreamInterpolate(FStateStreamInterpolateContext& Context, TArray<FTransform>& Out, const TArray<FTransform>& From, const TArray<FTransform>& To)
+{
+	int32 Num = To.Num();
+	check(Num == From.Num());
+	Out.SetNum(Num);
+	for (int32 I=0; I!=Num; ++I)
+	{
+		Out[I].Blend(From[I], To[I], Context.Factor);
+	}
+}
+
+inline bool StateStreamEquals(const TArray<FTransform>& A, const TArray<FTransform>& B)
+{
+	return false;
+}
+
+////////////////////////////////////////////////////////////////////////////////////////////////////
diff --git a/Engine/Source/Runtime/StateStream/StateStream.Build.cs b/Engine/Source/Runtime/StateStream/StateStream.Build.cs
new file mode 100644
--- /dev/null
+++ b/Engine/Source/Runtime/StateStream/StateStream.Build.cs
@@ -0,0 +1,16 @@
+// Copyright Epic Games, Inc. All Rights Reserved.
+
+using UnrealBuildTool;
+
+public class StateStream : ModuleRules
+{
+	public StateStream(ReadOnlyTargetRules Target) : base(Target)
+	{
+        PublicDependencyModuleNames.AddRange(
+            new string[] {
+                "Core",
+				"CoreUObject"
+            }
+        );
+    }
+}
