[StateStream]
* Fixed android compile error

diff --git a/Engine/Source/Runtime/StateStream/Public/GenericStateStream.h b/Engine/Source/Runtime/StateStream/Public/GenericStateStream.h
--- a/Engine/Source/Runtime/StateStream/Public/GenericStateStream.h
+++ b/Engine/Source/Runtime/StateStream/Public/GenericStateStream.h
@@ -33,7 +33,7 @@ template<typename Settings>
 class TStateStream : public IStateStream, public Settings::InterfaceType, public IStateStreamHandleOwner
 {
 public:
-	using InterfaceType = Settings::InterfaceType;
+	using InterfaceType = typename Settings::InterfaceType;
 	using FHandle = typename InterfaceType::Handle;
 	using FStaticState = typename InterfaceType::StaticState;
 	using FDynamicState = typename InterfaceType::DynamicState;
