[Statestream]
* Added StateStreamCreator type. This can be used to plugin custom statestreams into the statestream manager
* Fixed so DebugName can be set on generic state stream settings struct

diff --git a/Engine/Source/Runtime/StateStream/Private/StateStreamManagerImpl.cpp b/Engine/Source/Runtime/StateStream/Private/StateStreamManagerImpl.cpp
--- a/Engine/Source/Runtime/StateStream/Private/StateStreamManagerImpl.cpp
+++ b/Engine/Source/Runtime/StateStream/Private/StateStreamManagerImpl.cpp
@@ -1,6 +1,7 @@
 // Copyright Epic Games, Inc. All Rights Reserved.
 
 #include "StateStreamManagerImpl.h"
+#include "StateStreamCreator.h"
 #include "StateStreamDebugRenderer.h"
 
 ////////////////////////////////////////////////////////////////////////////////////////////////////
@@ -57,6 +58,7 @@ void FStateStreamManagerImpl::Render_Register(IStateStream& Stream, bool TakeOwn
 	{
 		StateStreamsLookup.SetNum(Id+1);
 	}
+	check(StateStreamsLookup[Id] == nullptr);
 	StateStreamsLookup[Id] = &Stream;
 }
 
@@ -115,6 +117,12 @@ void FStateStreamManagerImpl::Render_GarbageCollect(bool AsTask)
 	}
 }
 
+IStateStream* FStateStreamManagerImpl::Render_GetStream(uint32 Id)
+{
+	check(Id < uint32(StateStreamsLookup.Num()));
+	return StateStreamsLookup[Id];
+}
+
 void FStateStreamManagerImpl::Game_DebugRender(IStateStreamDebugRenderer& Renderer)
 {
 	for (const StateStreamRec& Rec : StateStreams)
@@ -135,3 +143,62 @@ FStateStreamManagerImpl::~FStateStreamManagerImpl()
 }
 
 ////////////////////////////////////////////////////////////////////////////////////////////////////
+
+FStateStreamCreator* FStateStreamCreator::First;
+
+FStateStreamCreator::FStateStreamCreator(FRegisterFunction&& InRegisterFunction, FUnregisterFunction&& InUnregisterFunction)
+:	RegisterFunction(MoveTemp(InRegisterFunction))
+,	UnregisterFunction(MoveTemp(InUnregisterFunction))
+{
+	if (!First)
+	{
+		First = this;
+		Prev = nullptr;
+		Next = nullptr;
+	}
+	else
+	{
+		Prev = nullptr;
+		Next = First;
+		Next->Prev = this;
+		First = this;
+	}
+}
+
+FStateStreamCreator::~FStateStreamCreator()
+{
+	if (Next)
+	{
+		Next->Prev = Prev;
+	}
+
+	if (Prev)
+	{
+		Prev->Next = Next;
+	}
+	else
+	{
+		First = Next;
+	}
+
+	Next = nullptr;
+	Prev = nullptr;
+}
+
+void FStateStreamCreator::RegisterStateStreams(const FStateStreamRegisterContext& Context)
+{
+	for (FStateStreamCreator* It=First; It; It=It->Next)
+	{
+		It->RegisterFunction(Context);
+	}
+}
+
+void FStateStreamCreator::UnregisterStateStreams(const FStateStreamUnregisterContext& Context)
+{
+	for (FStateStreamCreator* It=First; It; It=It->Next)
+	{
+		It->UnregisterFunction(Context);
+	}
+}
+
+////////////////////////////////////////////////////////////////////////////////////////////////////
diff --git a/Engine/Source/Runtime/StateStream/Public/GenericStateStream.h b/Engine/Source/Runtime/StateStream/Public/GenericStateStream.h
--- a/Engine/Source/Runtime/StateStream/Public/GenericStateStream.h
+++ b/Engine/Source/Runtime/StateStream/Public/GenericStateStream.h
@@ -18,7 +18,8 @@ struct TStateStreamSettings
 {
 	using InterfaceType = TInterfaceType;
 	using UserDataType = TUserDataType;
-
+	enum { Id = InterfaceType::Id };
+	static inline constexpr const TCHAR* DebugName = InterfaceType::Handle::DebugName;
 	static inline constexpr bool SkipCreatingDeletes = false;
 };
 
@@ -80,7 +81,8 @@ public:
 	uint32 GetUsedInstancesCount() const { return Instances.GetUsedCount(); }
 	uint32 GetUsedDynamicstatesCount() const { return DynamicStates.GetUsedCount(); }
 
-	~TStateStream();
+	TStateStream() = default;
+	~TStateStream() = default;
 
 private:
 	struct FTick;
@@ -129,6 +131,9 @@ private:
 	StateStreamTime RendTime = 0; // Last used time for rendering
 
 	TArray<FInstance*> DeferredDestroys;
+
+	TStateStream(const TStateStream&) = delete;
+	TStateStream& operator=(const TStateStream&) = delete;
 };
 
 
@@ -305,7 +310,7 @@ void* TStateStream<Settings>::Game_GetVoidPointer()
 template<typename Settings>
 uint32 TStateStream<Settings>::GetId()
 {
-	return Settings::InterfaceType::Id;
+	return Settings::Id;
 }
 
 template<typename Settings>
@@ -468,7 +473,7 @@ void TStateStream<Settings>::Render_GarbageCollect()
 template<typename Settings>
 const TCHAR* TStateStream<Settings>::GetDebugName()
 {
-	return FHandle::GetDebugName();
+	return Settings::DebugName;
 }
 
 template<typename Settings>
@@ -480,7 +485,7 @@ void TStateStream<Settings>::DebugRender(IStateStreamDebugRenderer& Renderer)
 	{
 		ModifiedCount = Tick->ModifiedInstances.CountSetBits();
 	}
-	DebugLine.Appendf(TEXT("%s   Num: %u  Changed: %u"), FHandle::GetDebugName(), Instances.GetUsedCount(), ModifiedCount);
+	DebugLine.Appendf(TEXT("%s   Num: %u  Changed: %u"), GetDebugName(), Instances.GetUsedCount(), ModifiedCount);
 	Renderer.DrawText(*DebugLine);
 }
 
@@ -624,9 +629,4 @@ void TStateStream<Settings>::MakeInternal(T& State)
 	}
 }
 
-template<typename Settings>
-TStateStream<Settings>::~TStateStream()
-{
-}
-
 ////////////////////////////////////////////////////////////////////////////////////////////////////
diff --git a/Engine/Source/Runtime/StateStream/Public/StateStreamCreator.h b/Engine/Source/Runtime/StateStream/Public/StateStreamCreator.h
new file mode 100644
--- /dev/null
+++ b/Engine/Source/Runtime/StateStream/Public/StateStreamCreator.h
@@ -0,0 +1,54 @@
+// Copyright Epic Games, Inc. All Rights Reserved.
+
+#pragma once
+
+#include "Templates/Function.h"
+
+#define UE_API STATESTREAM_API
+
+class FSceneInterface;
+class FStateStreamManagerImpl;
+
+////////////////////////////////////////////////////////////////////////////////////////////////////
+
+struct FStateStreamRegisterContext
+{
+	FStateStreamManagerImpl& Manager;
+	FSceneInterface* Scene;
+};
+
+////////////////////////////////////////////////////////////////////////////////////////////////////
+
+struct FStateStreamUnregisterContext
+{
+	FStateStreamManagerImpl& Manager;
+};
+
+////////////////////////////////////////////////////////////////////////////////////////////////////
+
+class FStateStreamCreator
+{
+public:
+	using FRegisterFunction = TFunction<void(const FStateStreamRegisterContext&)>;
+	using FUnregisterFunction = TFunction<void(const FStateStreamUnregisterContext&)>;
+	UE_API FStateStreamCreator(FRegisterFunction&& InRegisterFunction, FUnregisterFunction&& InUnregisterFunction);
+	UE_API ~FStateStreamCreator();
+
+	UE_API static void RegisterStateStreams(const FStateStreamRegisterContext&);
+	UE_API static void UnregisterStateStreams(const FStateStreamUnregisterContext&);
+
+private:
+	FStateStreamCreator(const FStateStreamCreator&) = delete;
+	FStateStreamCreator& operator=(const FStateStreamCreator&) = delete;
+
+	FRegisterFunction RegisterFunction;
+	FUnregisterFunction UnregisterFunction;
+
+	static FStateStreamCreator* First;
+	FStateStreamCreator* Next;
+	FStateStreamCreator* Prev;
+};
+
+////////////////////////////////////////////////////////////////////////////////////////////////////
+
+#undef UE_API
diff --git a/Engine/Source/Runtime/StateStream/Public/StateStreamHandle.h b/Engine/Source/Runtime/StateStream/Public/StateStreamHandle.h
--- a/Engine/Source/Runtime/StateStream/Public/StateStreamHandle.h
+++ b/Engine/Source/Runtime/StateStream/Public/StateStreamHandle.h
@@ -89,7 +89,7 @@ struct FStateStreamCopyContext
 	F##Type##Handle(FStateStreamCopyContext& C, const F##Type##Handle& O) : FStateStreamHandle(C, O) {} \
 	void Apply(FStateStreamCopyContext& Context, const FStateStreamHandle& Other) { FStateStreamHandle::Apply(Context, Other); } \
 	bool operator==(const F##Type##Handle& Other) const { return FStateStreamHandle::operator==(Other); } \
-	static const TCHAR* GetDebugName() { return TEXT(#Type); }; \
+	static inline constexpr const TCHAR* DebugName = TEXT(#Type); \
 	using FStateStreamHandle::Apply; \
 	friend F##Type##DynamicState; \
 
diff --git a/Engine/Source/Runtime/StateStream/Public/StateStreamManagerImpl.h b/Engine/Source/Runtime/StateStream/Public/StateStreamManagerImpl.h
--- a/Engine/Source/Runtime/StateStream/Public/StateStreamManagerImpl.h
+++ b/Engine/Source/Runtime/StateStream/Public/StateStreamManagerImpl.h
@@ -39,6 +39,9 @@ public:
 	// Garbage collect
 	UE_API void Render_GarbageCollect(bool AsTask = false);
 
+	// Get state stream from id
+	UE_API IStateStream* Render_GetStream(uint32 Id);
+
 	// Debug
 	UE_API virtual void Game_DebugRender(IStateStreamDebugRenderer& Renderer);
 
