[RadDebugger]
* New binaries that fixes a bunch of crashes related to modular builds. Thanks Nikita!

diff --git a/Engine/Source/Runtime/StateStream/Private/TransformStateStreamImpl.cpp b/Engine/Source/Runtime/StateStream/Private/TransformStateStreamImpl.cpp
--- a/Engine/Source/Runtime/StateStream/Private/TransformStateStreamImpl.cpp
+++ b/Engine/Source/Runtime/StateStream/Private/TransformStateStreamImpl.cpp
@@ -1,6 +1,9 @@
 // Copyright Epic Games, Inc. All Rights Reserved.
 
 #include "TransformStateStreamImpl.h"
+#include "StateStreamCreator.h"
+
+////////////////////////////////////////////////////////////////////////////////////////////////////
 
 FTransformObject::~FTransformObject()
 {
@@ -81,6 +84,8 @@ void FTransformObject::OnTransformObjectDirty()
 	CallListeners();
 }
 
+////////////////////////////////////////////////////////////////////////////////////////////////////
+
 void FTransformStateStreamImpl::SetParent(FTransformObject* Object, const FTransformDynamicState& Ds)
 {
 	const FTransformHandle& ParentHandle = Ds.GetParent();
@@ -136,3 +141,11 @@ void FTransformStateStreamImpl::Render_OnDestroy(const FTransformStaticState& Ss
 {
 	UserData->Release();
 }
+
+////////////////////////////////////////////////////////////////////////////////////////////////////
+
+FStateStreamCreator TransformStateStreamCreator(TransformStateStreamId,
+	[](const FStateStreamRegisterContext& Context) { Context.Register(*new FTransformStateStreamImpl(), true); },
+	[](const FStateStreamUnregisterContext& Context) {});
+
+////////////////////////////////////////////////////////////////////////////////////////////////////
