[StateStream]
* Another iteration on StateStreamCreator code... coming changelists will fix so StateStreamCreator is used to instantiate all statestreams

diff --git a/Engine/Source/Runtime/StateStream/Private/StateStreamManagerImpl.cpp b/Engine/Source/Runtime/StateStream/Private/StateStreamManagerImpl.cpp
--- a/Engine/Source/Runtime/StateStream/Private/StateStreamManagerImpl.cpp
+++ b/Engine/Source/Runtime/StateStream/Private/StateStreamManagerImpl.cpp
@@ -212,4 +212,14 @@ void FStateStreamCreator::UnregisterStateStreams(const FStateStreamUnregisterCon
 	}
 }
 
+void FStateStreamRegisterContext::Register(IStateStream& StateStream, bool TakeOwnership) const
+{
+	Manager.Render_Register(StateStream, TakeOwnership);
+}
+
+void FStateStreamRegisterContext::RegisterDependency(uint32 FromId, uint32 ToId) const
+{
+	Manager.Render_RegisterDependency(FromId, ToId);
+}
+
 ////////////////////////////////////////////////////////////////////////////////////////////////////
diff --git a/Engine/Source/Runtime/StateStream/Private/TransformStateStreamImpl.cpp b/Engine/Source/Runtime/StateStream/Private/TransformStateStreamImpl.cpp
--- a/Engine/Source/Runtime/StateStream/Private/TransformStateStreamImpl.cpp
+++ b/Engine/Source/Runtime/StateStream/Private/TransformStateStreamImpl.cpp
@@ -1,6 +1,9 @@
 // Copyright Epic Games, Inc. All Rights Reserved.
 
 #include "TransformStateStreamImpl.h"
+#include "StateStreamCreator.h"
+
+////////////////////////////////////////////////////////////////////////////////////////////////////
 
 FTransformObject::~FTransformObject()
 {
@@ -81,6 +84,8 @@ void FTransformObject::OnTransformObjectDirty()
 	CallListeners();
 }
 
+////////////////////////////////////////////////////////////////////////////////////////////////////
+
 void FTransformStateStreamImpl::SetParent(FTransformObject* Object, const FTransformDynamicState& Ds)
 {
 	const FTransformHandle& ParentHandle = Ds.GetParent();
@@ -136,3 +141,11 @@ void FTransformStateStreamImpl::Render_OnDestroy(const FTransformStaticState& Ss
 {
 	UserData->Release();
 }
+
+////////////////////////////////////////////////////////////////////////////////////////////////////
+
+FStateStreamCreator TransformStateStreamCreator(TransformStateStreamId,
+	[](const FStateStreamRegisterContext& Context) { Context.Register(*new FTransformStateStreamImpl(), true); },
+	[](const FStateStreamUnregisterContext& Context) {});
+
+////////////////////////////////////////////////////////////////////////////////////////////////////
diff --git a/Engine/Source/Runtime/StateStream/Public/GenericStateStream.h b/Engine/Source/Runtime/StateStream/Public/GenericStateStream.h
--- a/Engine/Source/Runtime/StateStream/Public/GenericStateStream.h
+++ b/Engine/Source/Runtime/StateStream/Public/GenericStateStream.h
@@ -39,6 +39,7 @@ public:
 	using FStaticState = typename InterfaceType::StaticState;
 	using FDynamicState = typename InterfaceType::DynamicState;
 	using FUserDataType = typename Settings::UserDataType;
+	enum { Id = Settings::Id };
 
 	// IStateStreamOwner (used by state stream handles on game side)
 	virtual FHandle Game_CreateInstance(const FStaticState& Ss, const FDynamicState& Ds) override final;
@@ -310,7 +311,7 @@ void* TStateStream<Settings>::Game_GetVoidPointer()
 template<typename Settings>
 uint32 TStateStream<Settings>::GetId()
 {
-	return Settings::Id;
+	return Id;
 }
 
 template<typename Settings>
diff --git a/Engine/Source/Runtime/StateStream/Public/StateStreamCreator.h b/Engine/Source/Runtime/StateStream/Public/StateStreamCreator.h
--- a/Engine/Source/Runtime/StateStream/Public/StateStreamCreator.h
+++ b/Engine/Source/Runtime/StateStream/Public/StateStreamCreator.h
@@ -8,6 +8,7 @@
 
 class FSceneInterface;
 class FStateStreamManagerImpl;
+class IStateStream;
 
 ////////////////////////////////////////////////////////////////////////////////////////////////////
 
@@ -15,6 +16,9 @@ struct FStateStreamRegisterContext
 {
 	FStateStreamManagerImpl& Manager;
 	FSceneInterface* Scene;
+
+	UE_API void Register(IStateStream& StateStream, bool TakeOwnership) const;
+	UE_API void RegisterDependency(uint32 FromId, uint32 ToId) const;
 };
 
 ////////////////////////////////////////////////////////////////////////////////////////////////////
@@ -31,9 +35,11 @@ class FStateStreamCreator
 public:
 	using FRegisterFunction = TFunction<void(const FStateStreamRegisterContext&)>;
 	using FUnregisterFunction = TFunction<void(const FStateStreamUnregisterContext&)>;
+
 	UE_API FStateStreamCreator(uint32 Id, FRegisterFunction&& InRegisterFunction, FUnregisterFunction&& InUnregisterFunction);
 	UE_API ~FStateStreamCreator();
 
+	// Called by system owning StateStreamManager
 	UE_API static void RegisterStateStreams(const FStateStreamRegisterContext&);
 	UE_API static void UnregisterStateStreams(const FStateStreamUnregisterContext&);
 
@@ -52,4 +58,24 @@ private:
 
 ////////////////////////////////////////////////////////////////////////////////////////////////////
 
+#define STATESTREAM_CREATOR_INSTANCE(ImplName) \
+	FStateStreamCreator ImplName##Creator(ImplName::Id, \
+		[](const FStateStreamRegisterContext& Context) \
+			{ ImplName& Impl = *new ImplName(*Context.Scene); Context.Register(Impl, true); }, \
+		[](const FStateStreamUnregisterContext& Context) {});
+
+#define STATESTREAM_CREATOR_INSTANCE_WITH_DEPENDENCY(ImplName, ToId) \
+	FStateStreamCreator ImplName##Creator(ImplName::Id, \
+		[](const FStateStreamRegisterContext& Context) \
+			{ ImplName& Impl = *new ImplName(*Context.Scene); Context.Register(Impl, true); Context.RegisterDependency(ImplName::Id, ToId); }, \
+		[](const FStateStreamUnregisterContext& Context) {});
+
+#define STATESTREAM_CREATOR_INSTANCE_WITH_FUNC(ImplName, ...) \
+	FStateStreamCreator ImplName##Creator(ImplName::Id, \
+		[](const FStateStreamRegisterContext& Context) \
+			{ ImplName& Impl = *new ImplName(*Context.Scene); Context.Register(Impl, true); __VA_ARGS__(Context, Impl); }, \
+		[](const FStateStreamUnregisterContext& Context) {});
+
+////////////////////////////////////////////////////////////////////////////////////////////////////
+
 #undef UE_API
