[StateStream]
* Added float and double to interpolation functions.

diff --git a/Engine/Source/Runtime/StateStream/Public/StateStreamDefinitions.h b/Engine/Source/Runtime/StateStream/Public/StateStreamDefinitions.h
--- a/Engine/Source/Runtime/StateStream/Public/StateStreamDefinitions.h
+++ b/Engine/Source/Runtime/StateStream/Public/StateStreamDefinitions.h
@@ -43,6 +43,15 @@ inline void StateStreamInterpolate(FStateStreamInterpolateContext& Context, uint
 	Out = From + uint32(double(To - From)*Context.Factor);
 }
 
+inline void StateStreamInterpolate(FStateStreamInterpolateContext& Context, float& Out, const float& From, const float& To)
+{
+	Out = From + float(double(To - From)*Context.Factor);
+}
+
+inline void StateStreamInterpolate(FStateStreamInterpolateContext& Context, double& Out, const double& From, const double& To)
+{
+	Out = From + (To - From)*Context.Factor;
+}
 
 template<typename Type>
 inline bool StateStreamEquals(const Type& A, const Type& B)
