[StateStream]
* Added TimeFactor and UserData parameters to state stream handle update. TimeFactor is not implemented yet. UserData will be used by light statestream backend

diff --git a/Engine/Source/Runtime/StateStream/Public/GenericStateStream.h b/Engine/Source/Runtime/StateStream/Public/GenericStateStream.h
--- a/Engine/Source/Runtime/StateStream/Public/GenericStateStream.h
+++ b/Engine/Source/Runtime/StateStream/Public/GenericStateStream.h
@@ -41,12 +41,14 @@ public:
 	using FUserDataType = typename Settings::UserDataType;
 	enum { Id = Settings::Id };
 
-	// IStateStreamOwner (used by state stream handles on game side)
+	// InterfaceType
 	virtual FHandle Game_CreateInstance(const FStaticState& Ss, const FDynamicState& Ds) override final;
+
+	// IStateStreamHandleOwner (used by state stream handles on game side)
 	virtual void Game_AddRef(uint32 HandleId) override final;
 	virtual void Game_Release(uint32 HandleId) override final;
-	virtual void Game_Update(uint32 HandleId, const void* Ds) override final;
-	virtual void* Game_Edit(uint32 HandleId) override final;
+	virtual void Game_Update(uint32 HandleId, const void* Ds, double TimeFactor, uint64 UserData) override;
+	virtual void* Game_Edit(uint32 HandleId, double TimeFactor, uint64 UserData) override;
 	virtual void* Render_GetUserData(uint32 HandleId) override final;
 	
 	// IStateStream (used by StateStreamManagerImpl)
@@ -85,6 +87,12 @@ public:
 	TStateStream() = default;
 	~TStateStream() = default;
 
+protected:
+	FDynamicState& Edit(uint32 HandleId, double TimeFactor);
+
+	template<typename InDynamicState>
+	void Update(uint32 HandleId, const InDynamicState& Ds, double TimeFactor);
+
 private:
 	struct FTick;
 
@@ -203,16 +211,23 @@ void TStateStream<Settings>::Game_Release(uint32 HandleId)
 }
 
 template<typename Settings>
-void TStateStream<Settings>::Game_Update(uint32 HandleId, const void* Ds)
+void TStateStream<Settings>::Game_Update(uint32 HandleId, const void* Ds, double TimeFactor, uint64 UserData)
 {
-	FDynamicState* DsState = static_cast<FDynamicState*>(Game_Edit(HandleId));
+	FDynamicState& DsState = Edit(HandleId, TimeFactor);
 	FStateStreamCopyContext Context;
 	Context.IsInternal = 0;
-	DsState->Apply(Context, *static_cast<const FDynamicState*>(Ds));
+	DsState.Apply(Context, *static_cast<const FDynamicState*>(Ds));
 }
 
 template<typename Settings>
-void* TStateStream<Settings>::Game_Edit(uint32 HandleId)
+void* TStateStream<Settings>::Game_Edit(uint32 HandleId, double TimeFactor, uint64 UserData)
+{
+	return &Edit(HandleId, TimeFactor);
+}
+
+
+template<typename Settings>
+typename TStateStream<Settings>::FDynamicState& TStateStream<Settings>::Edit(uint32 HandleId, double TimeFactor)
 {
 	check(HandleId != 0);
 	uint32 InstanceIndex = HandleId - 1;
@@ -244,9 +259,33 @@ void* TStateStream<Settings>::Game_Edit(uint32 HandleId)
 		DsState = new (DsPtr) FDynamicState(OldDs);
 		MakeInternal(OldDs);
 	}
-	return DsState;
+	return *DsState;
 }
 
+template<typename Settings>
+template<typename InDynamicState>
+void TStateStream<Settings>::Update(uint32 HandleId, const InDynamicState& Ds, double TimeFactor)
+{
+	/*
+	check(HandleId != 0);
+	uint32 InstanceIndex = HandleId - 1;
+
+	check(CurrentTick);
+	check(InstanceIndex < uint32(CurrentTick->DynamicStates.Num()));
+
+	CurrentTickLock.ReadLock();
+	uint32 DsIndex = CurrentTick->DynamicStates[InstanceIndex];
+	auto& CurrentState = *(decltype(InDynamicState)*)DynamicStates[DsIndex];
+	if (!CurrentState.Diff(InDs)) // TODO
+	{
+		return;
+	}
+	*/
+	FStateStreamCopyContext Context;
+	Context.IsInternal = 0;
+	InDynamicState& DsState = Edit(HandleId, TimeFactor);
+	DsState.Apply(Context, Ds);
+}
 
 template<typename Settings>
 void* TStateStream<Settings>::Render_GetUserData(uint32 HandleId)
diff --git a/Engine/Source/Runtime/StateStream/Public/StateStreamHandle.h b/Engine/Source/Runtime/StateStream/Public/StateStreamHandle.h
--- a/Engine/Source/Runtime/StateStream/Public/StateStreamHandle.h
+++ b/Engine/Source/Runtime/StateStream/Public/StateStreamHandle.h
@@ -58,8 +58,8 @@ class IStateStreamHandleOwner
 public:
 	virtual void Game_AddRef(uint32 HandleId) = 0;
 	virtual void Game_Release(uint32 HandleId) = 0;
-	virtual void Game_Update(uint32 HandleId, const void* DynamicState) = 0;
-	virtual void* Game_Edit(uint32 HandleId) = 0;
+	virtual void Game_Update(uint32 HandleId, const void* DynamicState, double TimeFactor, uint64 UserData) = 0;
+	virtual void* Game_Edit(uint32 HandleId, double TimeFactor, uint64 UserData) = 0;
 	virtual void* Render_GetUserData(uint32 HandleId) = 0;
 };
 
@@ -83,8 +83,8 @@ struct FStateStreamCopyContext
 	F##Type##Handle& operator=(F##Type##Handle&& Other) { FStateStreamHandle::operator=(std::move(Other)); return *this; } \
 	F##Type##Handle& operator=(const F##Type##Handle& Other) { FStateStreamHandle::operator=(Other); return *this; } \
 	~F##Type##Handle() = default; \
-	void Update(const F##Type##DynamicState& Ds) { if (Owner) Owner->Game_Update(Id, &Ds); } \
-	F##Type##DynamicState& Edit() { return *static_cast<F##Type##DynamicState*>(Owner->Game_Edit(Id)); } \
+	void Update(const F##Type##DynamicState& Ds, double TimeFactor = 0) { if (Owner) Owner->Game_Update(Id, &Ds, TimeFactor, 0); } \
+	F##Type##DynamicState& Edit(double TimeFactor = 0) { return *static_cast<F##Type##DynamicState*>(Owner->Game_Edit(Id, TimeFactor, 0)); } \
 	F##Type##Handle(IStateStreamHandleOwner& O, uint32 I) : FStateStreamHandle(O, I) {} \
 	F##Type##Handle(FStateStreamCopyContext& C, const F##Type##Handle& O) : FStateStreamHandle(C, O) {} \
 	void Apply(FStateStreamCopyContext& Context, const FStateStreamHandle& Other) { FStateStreamHandle::Apply(Context, Other); } \
