[StateStream]
* Changed std::move to MoveTemp

diff --git a/Engine/Source/Runtime/StateStream/Public/Pow2ChunkedArray.h b/Engine/Source/Runtime/StateStream/Public/Pow2ChunkedArray.h
--- a/Engine/Source/Runtime/StateStream/Public/Pow2ChunkedArray.h
+++ b/Engine/Source/Runtime/StateStream/Public/Pow2ChunkedArray.h
@@ -92,7 +92,7 @@ template<typename T, uint32 MinSize, uint32 MaxSize>
 T& FPow2ChunkedArray<T, MinSize, MaxSize>::Add(T&& Value)
 {
 	uint32 OutIndex;
-	return *new (AddUninitialized(OutIndex)) T(std::move(Value));
+	return *new (AddUninitialized(OutIndex)) T(MoveTemp(Value));
 }
 
 template<typename T, uint32 MinSize, uint32 MaxSize>
