[Statestream]
* Added feature to be able to toggle render side on/off
* Fixed so SkinnedMeshStateStream works with new changes in SkeletalMeshUpdater

#Synced-CL 49872240

diff --git a/Engine/Source/Runtime/Renderer/Private/StateStream/ParticleSystemStateStreamImpl.h b/Engine/Source/Runtime/Renderer/Private/StateStream/ParticleSystemStateStreamImpl.h
--- a/Engine/Source/Runtime/Renderer/Private/StateStream/ParticleSystemStateStreamImpl.h
+++ b/Engine/Source/Runtime/Renderer/Private/StateStream/ParticleSystemStateStreamImpl.h
@@ -148,6 +148,7 @@ public:
 	virtual void Render_PostUpdate() override {}
 	virtual void Render_Exit() override {}
 	virtual void Render_GarbageCollect() override {}
+	virtual void Render_Enable(bool Enable) override {};
 	virtual uint32 GetId() override { return ParticleSystemStateStreamId; }
 
 	FParticleSystemStateStreamCascade* CascadeBackend = nullptr;
diff --git a/Engine/Source/Runtime/Renderer/Private/StateStream/SkinnedMeshStateStreamImpl.cpp b/Engine/Source/Runtime/Renderer/Private/StateStream/SkinnedMeshStateStreamImpl.cpp
--- a/Engine/Source/Runtime/Renderer/Private/StateStream/SkinnedMeshStateStreamImpl.cpp
+++ b/Engine/Source/Runtime/Renderer/Private/StateStream/SkinnedMeshStateStreamImpl.cpp
@@ -5,6 +5,7 @@
 #include "Rendering/RenderCommandPipes.h"
 #include "ScenePrivate.h"
 #include "SkeletalMeshSceneProxy.h"
+#include "SkeletalMeshUpdater.h"
 #include "SkeletalRenderPublic.h"
 #include "SkinnedMeshSceneProxyDesc.h"
 #include "StateStreamCreator.h"
@@ -201,10 +202,20 @@ void FSkinnedMeshStateStreamImpl::Render_OnDestroy(const FSkinnedMeshStaticState
 {
 	if (Object)
 	{
+		Object->MeshObject->ReleaseResources();
 		Object->Release();
+		Object = nullptr;
 	}
 }
 
+void FSkinnedMeshStateStreamImpl::Render_PostUpdate()
+{
+	TStateStream<FSkinnedMeshStateStreamSettings>::Render_PostUpdate();
+
+	// TODO: Get rid of this entirely.. no need to create these queues of stuff when it can be done directly render side
+	Scene.GetSkeletalMeshUpdater()->PushChangesToStream();
+}
+
 ////////////////////////////////////////////////////////////////////////////////////////////////////
 
 STATESTREAM_CREATOR_INSTANCE_WITH_DEPENDENCY(FSkinnedMeshStateStreamImpl, TransformStateStreamId)
diff --git a/Engine/Source/Runtime/Renderer/Private/StateStream/SkinnedMeshStateStreamImpl.h b/Engine/Source/Runtime/Renderer/Private/StateStream/SkinnedMeshStateStreamImpl.h
--- a/Engine/Source/Runtime/Renderer/Private/StateStream/SkinnedMeshStateStreamImpl.h
+++ b/Engine/Source/Runtime/Renderer/Private/StateStream/SkinnedMeshStateStreamImpl.h
@@ -60,6 +60,9 @@ public:
 	FSkinnedMeshStateStreamImpl(FSceneInterface& InScene);
 private:
 	void SetTransformObject(FSkinnedMeshObject& Object, const FSkinnedMeshDynamicState& Ds);
+
+	UE_API virtual void Render_PostUpdate() override;
+
 	UE_API virtual void Render_OnCreate(const FSkinnedMeshStaticState& Ss, const FSkinnedMeshDynamicState& Ds, FSkinnedMeshObject*& UserData, bool IsDestroyedInSameFrame) override;
 	UE_API virtual void Render_OnUpdate(const FSkinnedMeshStaticState& Ss, const FSkinnedMeshDynamicState& Ds, FSkinnedMeshObject*& UserData) override;
 	UE_API virtual void Render_OnDestroy(const FSkinnedMeshStaticState& Ss, const FSkinnedMeshDynamicState& Ds, FSkinnedMeshObject*& UserData) override;
diff --git a/Engine/Source/Runtime/Renderer/Private/StateStream/StaticMeshStateStreamImpl.cpp b/Engine/Source/Runtime/Renderer/Private/StateStream/StaticMeshStateStreamImpl.cpp
--- a/Engine/Source/Runtime/Renderer/Private/StateStream/StaticMeshStateStreamImpl.cpp
+++ b/Engine/Source/Runtime/Renderer/Private/StateStream/StaticMeshStateStreamImpl.cpp
@@ -125,6 +125,7 @@ void FStaticMeshStateStreamImpl::Render_OnDestroy(const FStaticMeshStaticState&
 	if (UserData)
 	{
 		UserData->Release();
+		UserData = nullptr;
 	}
 }
 
diff --git a/Engine/Source/Runtime/StateStream/Private/StateStreamManagerImpl.cpp b/Engine/Source/Runtime/StateStream/Private/StateStreamManagerImpl.cpp
--- a/Engine/Source/Runtime/StateStream/Private/StateStreamManagerImpl.cpp
+++ b/Engine/Source/Runtime/StateStream/Private/StateStreamManagerImpl.cpp
@@ -4,6 +4,8 @@
 #include "StateStreamCreator.h"
 #include "StateStreamDebugRenderer.h"
 
+#define STATE_STREAM_DEBUG_CREATE_DESTROY 0
+
 ////////////////////////////////////////////////////////////////////////////////////////////////////
 
 void FStateStreamManagerImpl::Game_BeginTick()
@@ -74,6 +76,15 @@ void FStateStreamManagerImpl::Render_RegisterDependency(IStateStream& From, ISta
 
 void FStateStreamManagerImpl::Render_Update(double AbsoluteTime)
 {
+#if STATE_STREAM_DEBUG_CREATE_DESTROY
+	static uint32 Counter = 0;
+	++Counter;
+	for (const StateStreamRec& Rec : StateStreams)
+	{
+		Rec.Stream->Render_Enable(((Counter/100) % 2) ? true : false);
+	}
+#endif
+
 	for (const StateStreamRec& Rec : StateStreams)
 	{
 		Rec.Stream->Render_Update(AbsoluteTime);
diff --git a/Engine/Source/Runtime/StateStream/Private/TransformStateStreamImpl.cpp b/Engine/Source/Runtime/StateStream/Private/TransformStateStreamImpl.cpp
--- a/Engine/Source/Runtime/StateStream/Private/TransformStateStreamImpl.cpp
+++ b/Engine/Source/Runtime/StateStream/Private/TransformStateStreamImpl.cpp
@@ -140,6 +140,7 @@ void FTransformStateStreamImpl::Render_OnUpdate(const FTransformStaticState& Ss,
 void FTransformStateStreamImpl::Render_OnDestroy(const FTransformStaticState& Ss, const FTransformDynamicState& Ds, FTransformObject*& UserData)
 {
 	UserData->Release();
+	UserData = nullptr;
 }
 
 ////////////////////////////////////////////////////////////////////////////////////////////////////
diff --git a/Engine/Source/Runtime/StateStream/Public/GenericStateStream.h b/Engine/Source/Runtime/StateStream/Public/GenericStateStream.h
--- a/Engine/Source/Runtime/StateStream/Public/GenericStateStream.h
+++ b/Engine/Source/Runtime/StateStream/Public/GenericStateStream.h
@@ -61,6 +61,7 @@ public:
 	virtual void Render_PostUpdate() override;
 	virtual void Render_Exit() override;
 	virtual void Render_GarbageCollect() override;
+	virtual void Render_Enable(bool enable) override;
 
 	virtual const TCHAR* GetDebugName() override;
 	virtual void DebugRender(IStateStreamDebugRenderer& Renderer) override;
@@ -138,6 +139,8 @@ private:
 
 	FTick* RendTick = nullptr; // Last used tick for rendering.
 	StateStreamTime RendTime = 0; // Last used time for rendering
+	bool RendEnabledRequested = true;
+	bool RendEnabled = true;
 
 	TArray<FInstance*> DeferredDestroys;
 
@@ -356,8 +359,69 @@ uint32 TStateStream<Settings>::GetId()
 template<typename Settings>
 void TStateStream<Settings>::Render_Update(StateStreamTime AbsoluteTime)
 {
+	bool Requested = RendEnabledRequested;
+	if (Requested != RendEnabled)
+	{
+		RendEnabled = Requested;
+		if (RendEnabled)
+		{
+			// Figure out which tick to use and recreate instances
+			check(NewestAvailableTick);
+			FTick* Tick = NewestAvailableTick;
+			while (true)
+			{
+				if (AbsoluteTime >= Tick->Time || !Tick->PrevTick)
+				{
+					break;
+				}
+				Tick = Tick->PrevTick;
+			}
+
+			RendTick = Tick;
+			RendTime = AbsoluteTime;
+
+			for (uint32 InstanceIndex = 0, E = RendTick->DynamicStates.Num(); InstanceIndex != E; ++InstanceIndex)
+			{
+				FInstance& Instance = Instances[InstanceIndex];
+				if (Instance.CreateTick > Tick->Index || Instance.DeleteTick <= Tick->Index)
+				{
+					continue;
+				}
+				check(!Instance.RendDynamicState.IsSet());
+				uint32 DynamicStateIndex = Tick->DynamicStates[InstanceIndex];
+				FStateStreamCopyContext Context;
+				Instance.RendDynamicState.Emplace(Context, DynamicStates[DynamicStateIndex]);
+				Render_OnCreateInline(Instance.StaticState, *Instance.RendDynamicState, Instance.UserData, false);
+			}
+		}
+		else if (RendTick)
+		{
+			// Cleanup all instances
+			for (uint32 InstanceIndex = 0, E = RendTick->DynamicStates.Num(); InstanceIndex != E; ++InstanceIndex)
+			{
+				FInstance& Instance = Instances[InstanceIndex];
+				if (!Instance.RendDynamicState)
+				{
+					continue;
+				}
+				Render_OnDestroyInline(Instance.StaticState, *Instance.RendDynamicState, Instance.UserData);
+				Instance.RendDynamicState.Reset();
+				checkf(!Instance.UserData, TEXT("Make sure to clear UserData in Render_OnDestroy"));
+			}
+			RendTick = nullptr;
+		}
+		return;
+	}
+	else if (!RendEnabled)
+	{
+		return;
+	}
+
 	if (!NewestAvailableTick || AbsoluteTime == RendTime)
+	{
 		return;
+	}
+
 	check(AbsoluteTime > RendTime); // Only play forward for now
 
 	uint32 PrevTick = 0;
@@ -429,6 +493,7 @@ void TStateStream<Settings>::Render_PostUpdate()
 	{
 		Render_OnDestroyInline(Instance->StaticState, *Instance->RendDynamicState, Instance->UserData);
 		Instance->RendDynamicState.Reset();
+		checkf(!Instance->UserData, TEXT("Make sure to clear UserData in Render_OnDestroy"));
 	}
 	DeferredDestroys.SetNum(0);
 }
@@ -456,6 +521,9 @@ void TStateStream<Settings>::Render_GarbageCollect()
 	if (RendTick && RendTick->Time != RendTime && RendTick->PrevTick)
 		RendTickUsed = RendTick->PrevTick;
 	
+	if (!RendTickUsed)
+		RendTickUsed = OldestAvailableTick->NextTick;
+
 	FTick* Tick = OldestAvailableTick;
 	while (Tick != RendTickUsed)
 	{
@@ -510,6 +578,12 @@ void TStateStream<Settings>::Render_GarbageCollect()
 
 }
 
+template<typename Settings>
+void TStateStream<Settings>::Render_Enable(bool Enable)
+{
+	RendEnabledRequested = Enable;
+}
+
 template<typename Settings>
 const TCHAR* TStateStream<Settings>::GetDebugName()
 {
@@ -655,6 +729,7 @@ void TStateStream<Settings>::ApplyChanges(const FTick& Tick, StateStreamTime Tim
 	{
 		Render_OnDestroyInline(Instance->StaticState, *Instance->RendDynamicState, Instance->UserData);
 		Instance->RendDynamicState.Reset();
+		checkf(!Instance->UserData, TEXT("Make sure to clear UserData in Render_OnDestroy"));
 	}
 }
 
diff --git a/Engine/Source/Runtime/StateStream/Public/StateStream.h b/Engine/Source/Runtime/StateStream/Public/StateStream.h
--- a/Engine/Source/Runtime/StateStream/Public/StateStream.h
+++ b/Engine/Source/Runtime/StateStream/Public/StateStream.h
@@ -22,6 +22,7 @@ public:
 	virtual void Render_PostUpdate() = 0;
 	virtual void Render_Exit() = 0;
 	virtual void Render_GarbageCollect() = 0;
+	virtual void Render_Enable(bool Enable) = 0;
 
 	virtual uint32 GetId() = 0;
 
