[StateStream]
* Fixed crash when starting statestream and also turning off render side happens at the same time

#Synced-CL 49900804

diff --git a/Engine/Source/Runtime/StateStream/Public/GenericStateStream.h b/Engine/Source/Runtime/StateStream/Public/GenericStateStream.h
--- a/Engine/Source/Runtime/StateStream/Public/GenericStateStream.h
+++ b/Engine/Source/Runtime/StateStream/Public/GenericStateStream.h
@@ -522,7 +522,13 @@ void TStateStream<Settings>::Render_GarbageCollect()
 		RendTickUsed = RendTick->PrevTick;
 	
 	if (!RendTickUsed)
+	{
 		RendTickUsed = OldestAvailableTick->NextTick;
+		if (!RendTickUsed)
+		{
+			return;
+		}
+	}
 
 	FTick* Tick = OldestAvailableTick;
 	while (Tick != RendTickUsed)
