Improve state stream debug rendering

#Synced-CL 49917485

diff --git a/Engine/Plugins/Runtime/nDisplay/Source/DisplayCluster/Private/Game/EngineClasses/Basics/DisplayClusterViewportClient.cpp b/Engine/Plugins/Runtime/nDisplay/Source/DisplayCluster/Private/Game/EngineClasses/Basics/DisplayClusterViewportClient.cpp
--- a/Engine/Plugins/Runtime/nDisplay/Source/DisplayCluster/Private/Game/EngineClasses/Basics/DisplayClusterViewportClient.cpp
+++ b/Engine/Plugins/Runtime/nDisplay/Source/DisplayCluster/Private/Game/EngineClasses/Basics/DisplayClusterViewportClient.cpp
@@ -49,6 +49,11 @@
 #include "Render/Viewport/IDisplayClusterViewportManager.h"
 #include "Render/Viewport/RenderFrame/DisplayClusterRenderFrame.h"
 
+#if WITH_STATE_STREAM
+#include "StateStreamDebugRenderer.h"
+#include "StateStreamManager.h"
+#endif
+
 
 // Debug feature to synchronize and force all external resources to be transferred cross GPU at the end of graph execution.
 // May be useful for testing cross GPU synchronization logic.
@@ -743,9 +748,9 @@ void UDisplayClusterViewportClient::Draw(FViewport* InViewport, FCanvas* SceneCa
 		public:
 			FStateStreamDebugRenderer(FViewport* V, UCanvas* C) : Viewport(V), Canvas(C) {}
 
-			virtual void DrawText(const FStringView& Text) override
+			virtual void DrawText(const FStringView& Text, const FColor& Color) override
 			{
-				Canvas->SetDrawColor(FColor::Orange);
+				Canvas->SetDrawColor(Color);
 				Canvas->DrawText(GEngine->GetSmallFont(), Text, 100, Y);
 				Y += 20;
 			}
diff --git a/Engine/Source/Runtime/Engine/Private/GameViewportClient.cpp b/Engine/Source/Runtime/Engine/Private/GameViewportClient.cpp
--- a/Engine/Source/Runtime/Engine/Private/GameViewportClient.cpp
+++ b/Engine/Source/Runtime/Engine/Private/GameViewportClient.cpp
@@ -1984,9 +1984,9 @@ void UGameViewportClient::Draw(FViewport* InViewport, FCanvas* SceneCanvas)
 		public:
 			FStateStreamDebugRenderer(FViewport* V, UCanvas* C) : Viewport(V), Canvas(C) {}
 
-			virtual void DrawText(const FStringView& Text) override
+			virtual void DrawText(const FStringView& Text, const FColor& Color) override
 			{
-				Canvas->SetDrawColor(FColor::Orange);
+				Canvas->SetDrawColor(Color);
 				Canvas->DrawText(GEngine->GetSmallFont(), Text, 100, Y);
 				Y += 20;
 			}
diff --git a/Engine/Source/Runtime/RenderCore/Private/RenderingThread.cpp b/Engine/Source/Runtime/RenderCore/Private/RenderingThread.cpp
--- a/Engine/Source/Runtime/RenderCore/Private/RenderingThread.cpp
+++ b/Engine/Source/Runtime/RenderCore/Private/RenderingThread.cpp
@@ -32,6 +32,7 @@
 #include "ProfilingDebugging/CsvProfiler.h"
 #include "Async/TaskTrace.h"
 #include "DataDrivenShaderPlatformInfo.h"
+#include "StateStreamManagerImpl.h"
 #include "HAL/ThreadManager.h"
 #include "ProfilingDebugging/CountersTrace.h"
 #include "Async/ParallelFor.h"
@@ -2424,6 +2425,18 @@ void FRenderCommandPipe::ExecuteCommands(UE::RenderCommandPipe::FCommandList& Co
 
 //////////////////////////////////////////////////////////////////////////
 
+#if WITH_STATE_STREAM
+void FRenderCommandDispatcher::RegisterDebugCommand(const FString CommandName)
+{
+	if (IsInGameThread() || IsInParallelGameThread())
+	{
+		FStateStreamManagerImpl::GetDebugRenderCommands().AddDebugElements(CommandName, 1);
+	}
+}
+#endif
+
+//////////////////////////////////////////////////////////////////////////
+
 TAutoConsoleVariable<int32> CVarAllowOneFrameThreadLag(
 	TEXT("r.OneFrameThreadLag"),
 	1,
diff --git a/Engine/Source/Runtime/RenderCore/Public/RenderingThread.h b/Engine/Source/Runtime/RenderCore/Public/RenderingThread.h
--- a/Engine/Source/Runtime/RenderCore/Public/RenderingThread.h
+++ b/Engine/Source/Runtime/RenderCore/Public/RenderingThread.h
@@ -858,7 +858,7 @@ class FRenderCommandList final : public TConcurrentLinearObject<FRenderCommandLi
 	friend class FRenderCommandPipe;
 	friend class FRenderThreadCommandPipe;
 	friend class FRenderCommandDispatcher;
-
+	
 	static thread_local FRenderCommandList* InstanceTLS;
 	static RENDERCORE_API FRenderCommandList* GetInstanceTLS();
 	static RENDERCORE_API FRenderCommandList* SetInstanceTLS(FRenderCommandList* CommandList);
@@ -1094,6 +1094,12 @@ private:
 
 class FRenderCommandDispatcher
 {
+private : 
+#if WITH_STATE_STREAM
+	/** Register the debug command to state stream */
+	RENDERCORE_API static void RegisterDebugCommand(const FString CommandName);
+#endif
+	
 public:
 	/**
 	 * Call to submit a command list into a parent command list or render command pipes. If the parent command list is null the recording instance
@@ -1108,7 +1114,9 @@ public:
 	template <typename RenderCommandTag>
 	static void Enqueue(TUniqueFunction<void(FRHICommandListImmediate&)>&& Function)
 	{
-#if !WITH_STATE_STREAM
+#if WITH_STATE_STREAM
+		RegisterDebugCommand(FString(RenderCommandTag::Get().GetName()));
+#else
 		if (FRenderCommandList* CommandList = FRenderCommandList::GetInstanceTLS())
 		{
 			CommandList->Enqueue<RenderCommandTag>(MoveTemp(Function));
@@ -1122,7 +1130,9 @@ public:
 	template <typename RenderCommandTag>
 	static void Enqueue(FRenderCommandPipe* Pipe, FRenderCommandPipe::FCommandListFunction&& Function)
 	{
-#if !WITH_STATE_STREAM
+#if WITH_STATE_STREAM
+		RegisterDebugCommand(FString(RenderCommandTag::Get().GetName()));
+#else
 		if (FRenderCommandList* CommandList = FRenderCommandList::GetInstanceTLS())
 		{
 			CommandList->Enqueue<RenderCommandTag>(Pipe, MoveTemp(Function));
@@ -1147,7 +1157,9 @@ public:
 	template <typename RenderCommandTag>
 	static void Enqueue(FRenderCommandPipe* Pipe, FRenderCommandPipe::FEmptyFunction&& Function)
 	{
-#if !WITH_STATE_STREAM
+#if WITH_STATE_STREAM
+		RegisterDebugCommand(FString(RenderCommandTag::Get().GetName()));
+#else
 		if (FRenderCommandList* CommandListSet = FRenderCommandList::GetInstanceTLS())
 		{
 			CommandListSet->Enqueue<RenderCommandTag>(Pipe, MoveTemp(Function));
diff --git a/Engine/Source/Runtime/RenderCore/RenderCore.Build.cs b/Engine/Source/Runtime/RenderCore/RenderCore.Build.cs
--- a/Engine/Source/Runtime/RenderCore/RenderCore.Build.cs
+++ b/Engine/Source/Runtime/RenderCore/RenderCore.Build.cs
@@ -36,7 +36,7 @@ public class RenderCore : ModuleRules
 			RuntimeDependencies.Add(Path.Combine(Unreal.EngineDirectory.ToString(), "Extras/GPUDumpViewer/..."), StagedFileType.DebugNonUFS);
 		}
 
-		PrivateDependencyModuleNames.AddRange(new string[] { "Core", "Projects", "ApplicationCore", "TraceLog", "CookOnTheFly" });
+		PrivateDependencyModuleNames.AddRange(new string[] { "Core", "Projects", "ApplicationCore", "TraceLog", "CookOnTheFly", "StateStream" });
 		
 		PublicIncludePathModuleNames.AddRange(new string[] { "RHI" });
 
diff --git a/Engine/Source/Runtime/Renderer/Private/RendererScene.cpp b/Engine/Source/Runtime/Renderer/Private/RendererScene.cpp
--- a/Engine/Source/Runtime/Renderer/Private/RendererScene.cpp
+++ b/Engine/Source/Runtime/Renderer/Private/RendererScene.cpp
@@ -1360,6 +1360,20 @@ template<class T>
 void FScene::BatchAddPrimitivesInternal(TArrayView<T*> InPrimitives)
 {
 	check(InPrimitives.Num() > 0);
+	
+#if WITH_STATE_STREAM
+	if (IsInGameThread() || IsInParallelGameThread())
+	{
+		if (FStateStreamManagerImpl* StateStreamManager = static_cast<FStateStreamManagerImpl*>(World->GetStateStreamManager()))
+		{
+			if (UObject* PrimitiveObject = ToUObject(InPrimitives.Last()))
+			{
+				const FString ObjectType = PrimitiveObject->GetClass()->GetFName().ToString();
+				StateStreamManager->GetDebugSceneProxies().AddDebugElements(ObjectType, InPrimitives.Num());
+			}
+		}
+	}
+#endif
 
 #if ENABLE_LOW_LEVEL_MEM_TRACKER
 	// If detailed per-tag asset memory stats are active, don't batch primitives, so the memory tags can be independent
diff --git a/Engine/Source/Runtime/StateStream/Private/StateStreamManagerImpl.cpp b/Engine/Source/Runtime/StateStream/Private/StateStreamManagerImpl.cpp
--- a/Engine/Source/Runtime/StateStream/Private/StateStreamManagerImpl.cpp
+++ b/Engine/Source/Runtime/StateStream/Private/StateStreamManagerImpl.cpp
@@ -6,8 +6,56 @@
 
 #define STATE_STREAM_DEBUG_CREATE_DESTROY 0
 
+void FStateStreamDebugElements::DrawDebugElements(const FString& DebugName, IStateStreamDebugRenderer& DebugRenderer) const
+{
+	TStringBuilder<1024> DebugLine;
+			
+	const int32 NumElementCounts = DebugElementsCounts.Num();
+	if (NumElementCounts == 0)
+	{
+		return;
+	}
+	const float NumElementColumns = FMath::Min(10.0f, NumElementCounts);
+	const int32 NumElementRows = FMath::CeilToInt32(NumElementCounts / NumElementColumns);
+	
+	// Display the number of elements
+	DebugLine.Reset();
+	DebugLine.Appendf(TEXT("Num %s = %d"), *DebugName, NumElementCounts);
+	DebugRenderer.DrawText(*DebugLine, FColor::Blue);
+	
+	// Display per element informations
+	auto ElementIterator = DebugElementsCounts.CreateConstIterator();
+	for (int32 ColumnIndex = 0; ColumnIndex < NumElementColumns; ++ColumnIndex)
+	{
+		DebugLine.Reset();
+		DebugLine.Appendf(TEXT("%s : "), *DebugName);
+		for (int32 RowIndex = 0; RowIndex < NumElementRows; ++RowIndex)
+		{
+			if (ElementIterator)
+			{
+				DebugLine.Appendf(TEXT("[%s | %d]"), *ElementIterator.Key(), ElementIterator.Value());
+				UE_LOG(LogTemp, Log, TEXT("%s : [%s | %d]"), *DebugName, *ElementIterator.Key(), ElementIterator.Value());
+				++ElementIterator;
+			}
+		}
+		DebugRenderer.DrawText(*DebugLine, FColor::Orange);
+	}
+}
+
 ////////////////////////////////////////////////////////////////////////////////////////////////////
 
+FStateStreamDebugElements FStateStreamManagerImpl::DebugRenderCommands;
+
+FStateStreamDebugElements& FStateStreamManagerImpl::GetDebugRenderCommands()
+{
+	return DebugRenderCommands;
+}
+
+FStateStreamDebugElements& FStateStreamManagerImpl::GetDebugSceneProxies()
+{
+	return DebugSceneProxies;
+}
+
 void FStateStreamManagerImpl::Game_BeginTick()
 {
 	check(!bIsInTick);
@@ -37,6 +85,11 @@ void FStateStreamManagerImpl::Game_Exit()
 		Rec.Stream->Game_Exit();
 	}
 	bGameExited = true;
+	
+	// Reset all the scene proxies at the end of the game since 
+	// we are monitoring all the registered proxies during 
+	// the full game life time
+	DebugSceneProxies.ResetDebugElements();
 }
 
 bool FStateStreamManagerImpl::Game_IsInTick()
@@ -136,10 +189,24 @@ IStateStream* FStateStreamManagerImpl::Render_GetStream(uint32 Id)
 
 void FStateStreamManagerImpl::Game_DebugRender(IStateStreamDebugRenderer& Renderer)
 {
+	TStringBuilder<1024> DebugLine;
+	
+	// Display the number of state streams
+	DebugLine.Reset();
+	DebugLine.Appendf(TEXT("Num State streams = %d"), StateStreams.Num());
+	Renderer.DrawText(*DebugLine, FColor::Blue);
+	
+	// Display per state stream informations
 	for (const StateStreamRec& Rec : StateStreams)
 	{
-		Rec.Stream->DebugRender(Renderer);
+		Rec.Stream->Game_DebugRender(Renderer);
 	}
+	// Draw the render commands and scene proxies debug data
+	DebugRenderCommands.DrawDebugElements(FString("Render command"), Renderer);
+	DebugSceneProxies.DrawDebugElements(FString("Scene proxies"), Renderer);
+	
+	// Reset the render commands debug elements every frame after using them
+	DebugRenderCommands.ResetDebugElements();
 }
 
 FStateStreamManagerImpl::~FStateStreamManagerImpl()
diff --git a/Engine/Source/Runtime/StateStream/Public/GenericStateStream.h b/Engine/Source/Runtime/StateStream/Public/GenericStateStream.h
--- a/Engine/Source/Runtime/StateStream/Public/GenericStateStream.h
+++ b/Engine/Source/Runtime/StateStream/Public/GenericStateStream.h
@@ -9,9 +9,24 @@
 #include "StateStreamDebugRenderer.h"
 #include "StateStreamDefinitions.h"
 #include "StateStreamStore.h"
+#include "ProfilingDebugging/CountersTrace.h"
+
+// State stream stat group
+DECLARE_STATS_GROUP(TEXT("StateStream"), STATGROUP_StateStream, STATCAT_Advanced);
+
+// Declare state stream counter to the stat system and to unreal insights 
+#define DECLARE_STATESTREAM_COUNTER_STATS(Name, Group) \
+	DECLARE_DWORD_COUNTER_STAT(TEXT(#Group) TEXT("|") TEXT(#Name), STAT_StateStream_Counters_##Group##_##Name, STATGROUP_StateStream) \
+	TRACE_DECLARE_INT_COUNTER(StateStream_Counters_##Group##_##Name, TEXT("StateStream/") TEXT(#Group) TEXT("/") TEXT(#Name));
+
+// Define state stream counter to the stat system and to unreal insights 
+#define DEFINE_STATESTREAM_COUNTER_STATS(Name, Group, Value) \
+	SET_DWORD_STAT(STAT_StateStream_Counters_##Group##_##Name, Value); \
+	TRACE_COUNTER_SET(StateStream_Counters_##Group##_##Name, Value)
 
 ////////////////////////////////////////////////////////////////////////////////////////////////////
-// Settings struct - Can be inherited to set settings on statestream
+// Settings struct - Can be inherited to set settings on state stream. It mainly contains all the data types used 
+// by state stream (static state / dynamic state / handle type / interface type) + id and debug name
 
 template<typename TInterfaceType, typename TUserDataType = void>
 struct TStateStreamSettings
@@ -23,7 +38,6 @@ struct TStateStreamSettings
 	static inline constexpr bool SkipCreatingDeletes = false;
 };
 
-
 ////////////////////////////////////////////////////////////////////////////////////////////////////
 // TStateStream is a generic implementation of IStateStream that contains all the boiler plate code
 // related to ticks, interpolation, etc. Should be the default goto implementation.
@@ -41,37 +55,40 @@ public:
 	using FUserDataType = typename Settings::UserDataType;
 	enum { Id = Settings::Id };
 
-	// InterfaceType
+	//~ Begin InterfaceType interface
 	virtual FHandle Game_CreateInstance(const FStaticState& Ss, const FDynamicState& Ds) override final;
+	//~ End InterfaceType interface
 
-	// IStateStreamHandleOwner (used by state stream handles on game side)
+	//~ Begin IStateStreamHandleOwner interface (used by state stream handles on game side)
 	virtual void Game_AddRef(uint32 HandleId) override final;
 	virtual void Game_Release(uint32 HandleId) override final;
 	virtual void Game_Update(uint32 HandleId, const void* Ds, double TimeFactor, uint64 UserData) override;
 	virtual void* Game_Edit(uint32 HandleId, double TimeFactor, uint64 UserData) override;
 	virtual void* Render_GetUserData(uint32 HandleId) override final;
+	//~ End IStateStreamHandleOwner interface
 	
-	// IStateStream (used by StateStreamManagerImpl)
+	//~ Begin IStateStream interface (used by StateStreamManagerImpl)
 	virtual void Game_BeginTick() override;
 	virtual void Game_EndTick(StateStreamTime AbsoluteTime) override;
 	virtual void Game_Exit() override;
 	virtual void* Game_GetVoidPointer() override;
-	virtual uint32 GetId() override;
+	virtual void Game_DebugRender(IStateStreamDebugRenderer& Renderer) override;
 	virtual void Render_Update(StateStreamTime AbsoluteTime) override;
 	virtual void Render_PostUpdate() override;
 	virtual void Render_Exit() override;
 	virtual void Render_GarbageCollect() override;
-	virtual void Render_Enable(bool enable) override;
-
+	virtual void Render_Enable(bool Enable) override;
 	virtual const TCHAR* GetDebugName() override;
-	virtual void DebugRender(IStateStreamDebugRenderer& Renderer) override;
-
+	virtual uint32 GetId() override;
+	//~ End IStateStream interface
 
+	// Get the instance user data on the render side
 	FUserDataType*& Render_GetUserData(const FHandle& Handle);
+	
+	// Get the instance interpolated dynamic state on the render side
 	const FDynamicState& Render_GetDynamicState(const FHandle& Handle);
-
-
-	// Specialize to do custom things
+	
+	// Specialize to do custom things when state has been created / updated or destroyed
 	virtual void Render_OnCreate(const FStaticState& Ss, const FDynamicState& Ds, FUserDataType*& UserData, bool IsDestroyedInSameFrame) {}
 	virtual void Render_OnUpdate(const FStaticState& Ss, const FDynamicState& Ds, FUserDataType*& UserData) {}
 	virtual void Render_OnDestroy(const FStaticState& Ss, const FDynamicState& Ds, FUserDataType*& UserData) {}
@@ -81,27 +98,17 @@ public:
 	void Render_OnUpdateInline(const FStaticState& Ss, const FDynamicState& Ds, FUserDataType*& UserData) { Render_OnUpdate(Ss, Ds, UserData); }
 	void Render_OnDestroyInline(const FStaticState& Ss, const FDynamicState& Ds, FUserDataType*& UserData) { Render_OnDestroy(Ss, Ds, UserData); }
 	
-	// For unit tests
+	// Get the used instance counts within the streams
 	uint32 GetUsedInstancesCount() const { return Instances.GetUsedCount(); }
-	uint32 GetUsedDynamicstatesCount() const { return DynamicStates.GetUsedCount(); }
+	
+	// Get the used dynamic states within the streams
+	uint32 GetUsedDynamicStatesCount() const { return DynamicStates.GetUsedCount(); }
 
 	TStateStream() = default;
 	~TStateStream() = default;
 
 protected:
-	FDynamicState& Edit(uint32 HandleId, double TimeFactor);
-
-	template<typename InDynamicState>
-	void Update(uint32 HandleId, const InDynamicState& Ds, double TimeFactor);
-
-private:
-	struct FTick;
-
-	void ApplyChanges(const FTick& Tick, StateStreamTime Time, uint32 PrevTickIndex, StateStreamTime PrevTime, const TBitArray<>& ModifiedInstances);
-
-	template<typename T>
-	void MakeInternal(T& State);
-
+	
 	// Information about instance.
 	struct FInstance
 	{
@@ -114,56 +121,88 @@ private:
 		TOptional<FDynamicState> RendDynamicState;
 		FUserDataType* UserData = nullptr;
 	};
-
-	TStateStreamStore<FInstance> Instances;
-
+	
 	// Tick produced by game side.
 	struct FTick
 	{
-		TArray<uint32> DynamicStates; // Contains index of all instances existing in tick (some might have been destroyed but index not reused)
+		TArray<uint32> DynamicStates; // Contains for each instance the index of the dynamic states used by this tick (some might have been destroyed but index not reused)
 		TBitArray<> ModifiedInstances; // Contains bits saying which of the instances that have been modified in this tick.
 		FTick* PrevTick = nullptr; // Tick with earlier time
 		FTick* NextTick = nullptr; // Tick with newer time
 		StateStreamTime Time = 0; // Time Tick finished
 		uint32 Index = 0; // Index of tick (created by TickCounter)
 	};
+	
+	// Get the dynamic state to be edited on GT in between begin/end tick
+	FDynamicState& Edit(uint32 HandleId, double TimeFactor);
+
+	// Update the editable dynamic state
+	template<typename InDynamicState>
+	void Update(uint32 HandleId, const InDynamicState& Ds, double TimeFactor);
 
+	// Interpolate and set the render dynamic states + call create / update / destroy callbacks. Called from the Render_Update
+	void ApplyChanges(const FTick& Tick, StateStreamTime Time, uint32 PrevTickIndex, StateStreamTime PrevTime, const TBitArray<>& ModifiedInstances);
+
+	// Make the state internal in case it has been released by its owner
+	template<typename T>
+	void MakeInternal(T& State);
+
+	// Lock in order to write into the current tick
 	FRWLock CurrentTickLock;
-	FTick* CurrentTick = nullptr; // Tick being worked on in game
-	FTick* OldestAvailableTick = nullptr; // Newest finished tick available to rendering
-	FTick* NewestAvailableTick = nullptr; // Newest finished tick available to rendering
+	
+	// Tick being worked on in game
+	FTick* CurrentTick = nullptr; 
+	
+	// Oldest finished tick available to rendering
+	FTick* OldestAvailableTick = nullptr; 
+	
+	// Newest finished tick available to rendering
+	FTick* NewestAvailableTick = nullptr; 
 
+	// Instances registered in the state stream 
+	TStateStreamStore<FInstance> Instances;
+	
+	// Dynamic states that will be referenced in the different ticks 
 	TStateStreamStore<FDynamicState> DynamicStates; // Store for dynamic states.
 
+	// Tick counter used to increment tick indices
 	uint32 TickCounter = 1;
 
-	FTick* RendTick = nullptr; // Last used tick for rendering.
-	StateStreamTime RendTime = 0; // Last used time for rendering
+	// Last used tick for rendering.
+	FTick* RendTick = nullptr; 
+	
+	// Last used time for rendering
+	StateStreamTime RendTime = 0;
+	
+	// Boolean to trigger the render update 
 	bool RendEnabledRequested = true;
+	
+	// Check if the render update has been enabled or not
 	bool RendEnabled = true;
 
+	// List of instances to be destroyed
 	TArray<FInstance*> DeferredDestroys;
 
 	TStateStream(const TStateStream&) = delete;
 	TStateStream& operator=(const TStateStream&) = delete;
 };
 
-
 ////////////////////////////////////////////////////////////////////////////////////////////////////
-// Implementation
 
 template<typename Settings>
 typename TStateStream<Settings>::FHandle TStateStream<Settings>::Game_CreateInstance(const FStaticState& Ss, const FDynamicState& Ds)
 {
 	check(CurrentTick);
 
-	uint32 InstanceIndex = Instances.Emplace(Ss, 1, CurrentTick->Index);
-
-	uint32 DynamicStateIndex = DynamicStates.Add(Ds);
+	// We usually have several dynamic states stored for each instance
+	// and consecutive  ticks could refer the same dynamic state if nothing changed
+	const uint32 InstanceIndex = Instances.Emplace(Ss, 1, CurrentTick->Index);
+	const uint32 DynamicStateIndex = DynamicStates.Add(Ds);
 
 	CurrentTickLock.WriteLock();
 
-	if (uint32(CurrentTick->DynamicStates.Num()) <= InstanceIndex)
+	// We make sure the current tick dynamic states / instances have at least the stream instances size
+	if (static_cast<uint32>(CurrentTick->DynamicStates.Num()) <= InstanceIndex)
 	{
 		CurrentTick->DynamicStates.SetNum(InstanceIndex + 1);
 		CurrentTick->ModifiedInstances.SetNum(InstanceIndex + 1, false);
@@ -200,6 +239,8 @@ void TStateStream<Settings>::Game_Release(uint32 HandleId)
 		return;
 	}
 
+	// If no more refs we are going to delete the tick store its index onto the delete tick
+	// the static and dynamic states should be made internal 
 	Instance.DeleteTick = CurrentTick->Index;
 	MakeInternal(Instance.StaticState);
 
@@ -228,7 +269,6 @@ void* TStateStream<Settings>::Game_Edit(uint32 HandleId, double TimeFactor, uint
 	return &Edit(HandleId, TimeFactor);
 }
 
-
 template<typename Settings>
 typename TStateStream<Settings>::FDynamicState& TStateStream<Settings>::Edit(uint32 HandleId, double TimeFactor)
 {
@@ -244,12 +284,16 @@ typename TStateStream<Settings>::FDynamicState& TStateStream<Settings>::Edit(uin
 
 	if (CurrentTick->ModifiedInstances[InstanceIndex])
 	{
+		// If the instance had already been modified just return the 
+		// current tick dynamic states for that instance
 		uint32 DsIndex = CurrentTick->DynamicStates[InstanceIndex];
 		CurrentTickLock.ReadUnlock();
 		DsState = &DynamicStates[DsIndex];
 	}
 	else
 	{
+		// If not modified, create a new dynamic state and reference 
+		// it in the current tick ones. 
 		uint32& IndexRef = CurrentTick->DynamicStates[InstanceIndex];
 		uint32 OldIndex = IndexRef;
 		CurrentTick->ModifiedInstances[InstanceIndex] = true;
@@ -260,6 +304,9 @@ typename TStateStream<Settings>::FDynamicState& TStateStream<Settings>::Edit(uin
 
 		FDynamicState& OldDs = DynamicStates[OldIndex];
 		DsState = new (DsPtr) FDynamicState(OldDs);
+		
+		// Why do wee need to make the old state internal since 
+		// this state could still be referenced by previous ticks?
 		MakeInternal(OldDs);
 	}
 	return *DsState;
@@ -286,6 +333,8 @@ void TStateStream<Settings>::Update(uint32 HandleId, const InDynamicState& Ds, d
 	*/
 	FStateStreamCopyContext Context;
 	Context.IsInternal = 0;
+	
+	// Get the editable state and apply the input dynamic state
 	InDynamicState& DsState = Edit(HandleId, TimeFactor);
 	DsState.Apply(Context, Ds);
 }
@@ -302,10 +351,14 @@ template<typename Settings>
 void TStateStream<Settings>::Game_BeginTick()
 {
 	check(!CurrentTick);
+	
+	// The new current tick for that stream is being created
 	CurrentTick = new FTick();
 	CurrentTick->Index = TickCounter++;
 	if (NewestAvailableTick)
 	{
+		// We copy bu default the list of dynamic states referenced by the last available tick 
+		// in case no instances will change and set the modified flag to false
 		CurrentTick->DynamicStates = NewestAvailableTick->DynamicStates;
 		CurrentTick->ModifiedInstances.SetNum(CurrentTick->DynamicStates.Num(), false);
 	}
@@ -328,6 +381,8 @@ void TStateStream<Settings>::Game_EndTick(StateStreamTime AbsoluteTime)
 		OldestAvailableTick = CurrentTick;
 	}
 		
+	// the current tick is being reset and the newest one available to rendering 
+	// is getting updated
 	NewestAvailableTick = CurrentTick;
 	CurrentTick = nullptr;
 }
@@ -363,6 +418,10 @@ void TStateStream<Settings>::Render_Update(StateStreamTime AbsoluteTime)
 	if (Requested != RendEnabled)
 	{
 		RendEnabled = Requested;
+		
+		// If the render update is enable and if it has changed we are in creation mode
+		// We need to find the tick with previous time lower than the requested absolute time
+		// and use it to fill the render dynamic state
 		if (RendEnabled)
 		{
 			// Figure out which tick to use and recreate instances
@@ -390,10 +449,14 @@ void TStateStream<Settings>::Render_Update(StateStreamTime AbsoluteTime)
 				check(!Instance.RendDynamicState.IsSet());
 				uint32 DynamicStateIndex = Tick->DynamicStates[InstanceIndex];
 				FStateStreamCopyContext Context;
+				
+				// Why do we use the render dynamic state from the prev tick and not interpolate from prev->next?
 				Instance.RendDynamicState.Emplace(Context, DynamicStates[DynamicStateIndex]);
 				Render_OnCreateInline(Instance.StaticState, *Instance.RendDynamicState, Instance.UserData, false);
 			}
 		}
+		// If the render update is disable and if it has changed we are in deletion mode
+		// We need to to reset all th instances render dynamic state
 		else if (RendTick)
 		{
 			// Cleanup all instances
@@ -597,7 +660,7 @@ const TCHAR* TStateStream<Settings>::GetDebugName()
 }
 
 template<typename Settings>
-void TStateStream<Settings>::DebugRender(IStateStreamDebugRenderer& Renderer)
+void TStateStream<Settings>::Game_DebugRender(IStateStreamDebugRenderer& Renderer)
 {
 	TStringBuilder<1024> DebugLine;
 	uint32 ModifiedCount = 0;
@@ -605,8 +668,8 @@ void TStateStream<Settings>::DebugRender(IStateStreamDebugRenderer& Renderer)
 	{
 		ModifiedCount = Tick->ModifiedInstances.CountSetBits();
 	}
-	DebugLine.Appendf(TEXT("%s   Num: %u  Changed: %u"), GetDebugName(), Instances.GetUsedCount(), ModifiedCount);
-	Renderer.DrawText(*DebugLine);
+	DebugLine.Appendf(TEXT("State stream : [%s | %u/%u modified]"), GetDebugName(), ModifiedCount, Instances.GetUsedCount());
+	Renderer.DrawText(*DebugLine, FColor::Orange);
 }
 
 template<typename Settings>
@@ -686,7 +749,7 @@ void TStateStream<Settings>::ApplyChanges(const FTick& Tick, StateStreamTime Tim
 			
 			StateStreamTime DeltaTime = Tick.Time - PrevTick.Time;
 			StateStreamTime TimeInTo = Time - PrevTick.Time;
-			double Factor = double(TimeInTo) / DeltaTime;
+			double Factor = (DeltaTime == 0.0) ? 0.0f : double(TimeInTo) / DeltaTime;
 
 			FDynamicState& From = DynamicStates[FromIndex];
 			FDynamicState& To = DynamicStates[ToIndex];
diff --git a/Engine/Source/Runtime/StateStream/Public/StateStream.h b/Engine/Source/Runtime/StateStream/Public/StateStream.h
--- a/Engine/Source/Runtime/StateStream/Public/StateStream.h
+++ b/Engine/Source/Runtime/StateStream/Public/StateStream.h
@@ -7,28 +7,48 @@
 class IStateStreamDebugRenderer;
 
 ////////////////////////////////////////////////////////////////////////////////////////////////////
-// Interface used by StateStreamManagerImpl. This is RT only. Should not be visible/used by GT
-// Documentation in StateStreamManagerImpl
+// Interface used by TStateStream to define all the GT/RT functions required 
+// to create/update/destroy the state stream on GT and to update/interpolate the result on RT
 
 class IStateStream
 {
 public:
+	// Create current tick data that will be update later on 
 	virtual void Game_BeginTick() = 0;
+	
+	// Make the current tick available to the renderer (update oldest/newest ticks) and then reset it 
 	virtual void Game_EndTick(StateStreamTime AbsoluteTime) = 0;
+	
+	// Release all the instance/ticks/dynamic states stored within that stream
 	virtual void Game_Exit() = 0;
+	
+	// Get the state stream interface pointer
 	virtual void* Game_GetVoidPointer() = 0;
+	
+	// Debug render streams data on GT 
+	virtual void Game_DebugRender(IStateStreamDebugRenderer& Renderer) {}
 
+	// Interpolate the instance render dynamic state from the available ticks [oldest-newest] given a time
 	virtual void Render_Update(StateStreamTime AbsoluteTime) = 0;
+	
+	// Destroy the updated render state
 	virtual void Render_PostUpdate() = 0;
+	
+	// Last update of the render dynamic state
 	virtual void Render_Exit() = 0;
+	
+	// Release all the ticks/data stored before the updated render tick
 	virtual void Render_GarbageCollect() = 0;
+	
+	// Enable the state stream render interpolation
 	virtual void Render_Enable(bool Enable) = 0;
 
+	// Get the state stream id
 	virtual uint32 GetId() = 0;
-
+	
+	// Get the state stream debug name
 	virtual const TCHAR* GetDebugName() { return TEXT("Unknown"); }
-	virtual void DebugRender(IStateStreamDebugRenderer& Renderer) {}
-
+	
 	virtual ~IStateStream() = default;
 };
 
diff --git a/Engine/Source/Runtime/StateStream/Public/StateStreamDebugRenderer.h b/Engine/Source/Runtime/StateStream/Public/StateStreamDebugRenderer.h
--- a/Engine/Source/Runtime/StateStream/Public/StateStreamDebugRenderer.h
+++ b/Engine/Source/Runtime/StateStream/Public/StateStreamDebugRenderer.h
@@ -3,13 +3,17 @@
 #pragma once
 
 #include "Containers/StringFwd.h"
+#include "Math/Color.h"
 
 ////////////////////////////////////////////////////////////////////////////////////////////////////
+// IStateStreamDebugRenderer interface to draw state streams debug data
 
 class IStateStreamDebugRenderer
 {
 public:
-	virtual void DrawText(const FStringView& Text) = 0;
+	// Interface to draw text with a given color for debug purposes
+	virtual void DrawText(const FStringView& Text, const FColor& Color) = 0;
+	virtual ~IStateStreamDebugRenderer() {}
 };
 
 ////////////////////////////////////////////////////////////////////////////////////////////////////
diff --git a/Engine/Source/Runtime/StateStream/Public/StateStreamManagerImpl.h b/Engine/Source/Runtime/StateStream/Public/StateStreamManagerImpl.h
--- a/Engine/Source/Runtime/StateStream/Public/StateStreamManagerImpl.h
+++ b/Engine/Source/Runtime/StateStream/Public/StateStreamManagerImpl.h
@@ -8,6 +8,44 @@
 
 #define UE_API STATESTREAM_API
 
+////////////////////////////////////////////////////////////////////////////////////////////////////
+// State stream debug elements to get information
+// about what remains to convert to state stream 
+
+class FStateStreamDebugElements
+{
+public:
+	// Add N debug elements given a type 
+	void AddDebugElements(const FString& ElementType, const uint32 NumElements)
+	{
+		if (uint32* ElementCount = DebugElementsCounts.Find(ElementType))
+		{
+			*ElementCount += NumElements;
+		}
+		else
+		{
+			DebugElementsCounts.Add(ElementType, NumElements);
+		}
+	}
+	
+	// Reset all the debug elements 
+	void ResetDebugElements()
+	{
+		DebugElementsCounts.Reset();
+	}
+	
+	// Draw the debug elements to the renderer 
+	void DrawDebugElements(const FString& DebubName, IStateStreamDebugRenderer& DebugRenderer) const;
+	
+	// Get the debug elements counts 
+	const TMap<FString,uint32>& GetDebugElements() const { return DebugElementsCounts; }
+
+private : 
+	
+	// Number of debug elements per type 
+	TMap<FString,uint32> DebugElementsCounts;
+};
+
 ////////////////////////////////////////////////////////////////////////////////////////////////////
 // StateStreamManager implementation. 
 // This type should only be known render side.
@@ -15,19 +53,23 @@
 class FStateStreamManagerImpl : public IStateStreamManager
 {
 public:
-	// IStateStreamManager interface
+	//~ Begin IStateStreamManager interface
 	UE_API virtual void Game_BeginTick() override;
 	UE_API virtual void Game_EndTick(double AbsoluteTime) override;
 	UE_API virtual void Game_Exit() override;
 	UE_API virtual bool Game_IsInTick() override;
 	UE_API virtual void* Game_GetStreamPointer(uint32 Id) override;
+	UE_API virtual void Game_DebugRender(IStateStreamDebugRenderer& Renderer) override;
+	//~ End IStateStreamManager interface
 
 	// Register new state streams into manager.
 	// TakeOwnership true means that manager will delete stream when shutting down
 	UE_API void Render_Register(IStateStream& Stream, bool TakeOwnership);
 
-	// Register dependency between statestreams. FromId will depend on ToId
+	// Register dependency between state streams given their ids. FromId will depend on ToId
 	UE_API void Render_RegisterDependency(uint32 FromId, uint32 ToId);
+	
+	// Register dependency between state streams given their interfaces. From will depend on To
 	UE_API void Render_RegisterDependency(IStateStream& From, IStateStream& To);
 
 	// Called at the beginning of a render frame. AbsolutTime is the amount of time the render frame consumes
@@ -36,30 +78,53 @@ public:
 	// Called before Render thread exits
 	UE_API void Render_Exit();
 
-	// Garbage collect
+	// Garbage collect all the unnecessary ticks/data within each streams
 	UE_API void Render_GarbageCollect(bool AsTask = false);
 
 	// Get state stream from id
 	UE_API IStateStream* Render_GetStream(uint32 Id);
-
-	// Debug
-	UE_API virtual void Game_DebugRender(IStateStreamDebugRenderer& Renderer);
-
-	UE_API ~FStateStreamManagerImpl();
-
+	
+	// Get the debug render commands 
+	UE_API static FStateStreamDebugElements& GetDebugRenderCommands();
+	
+	// Get the debug scene proxies 
+	UE_API FStateStreamDebugElements& GetDebugSceneProxies();
+	
+	UE_API virtual ~FStateStreamManagerImpl() override;
+	
 private:
+	
+	// Small structure used to list the state streams available in the manager 
 	struct StateStreamRec
 	{
 		IStateStream* Stream;
 		bool Owned;
 	};
+	
+	// List of state streams available in the manager 
 	TArray<StateStreamRec> StateStreams;
+	
+	// List of state streams for each registered state stream id. 
+	// This will allow fast lookup to retrieve a state stream given an id
 	TArray<IStateStream*> StateStreamsLookup;
+	
+	// Bool to check if we are currently in between the begin and the end tick
 	bool bIsInTick = false;
+	
+	// Bool to check if the game has already been exited on the GT side
 	bool bGameExited = false;
+	
+	// Bool to check if the game has already been exited on the RT side
 	bool bRenderExited = false;
 
+	// Async task to perform the garbage collection
 	UE::Tasks::FTask GarbageCollectTask;
+	
+	// Render commands debug data to keep track of the remaining commands to be converted to state stream
+	static FStateStreamDebugElements DebugRenderCommands;
+	
+	// Scene proxies debug data to keep track of the remaining proxies to be converted to state stream
+	FStateStreamDebugElements DebugSceneProxies;
 };
 
 ////////////////////////////////////////////////////////////////////////////////////////////////////
